Add a runtime mesh builder module, this module is use to build the render data of a mesh description

diff --git a/Engine/Source/Editor/UnrealEd/Private/Fbx/FbxStaticMeshImport.cpp b/Engine/Source/Editor/UnrealEd/Private/Fbx/FbxStaticMeshImport.cpp
--- a/Engine/Source/Editor/UnrealEd/Private/Fbx/FbxStaticMeshImport.cpp
+++ b/Engine/Source/Editor/UnrealEd/Private/Fbx/FbxStaticMeshImport.cpp
@@ -34,7 +34,7 @@
 #include "Misc/FbxErrors.h"
 #include "PhysicsEngine/BodySetup.h"
 #include "MeshDescription.h"
-#include "IMeshDescriptionModule.h"
+#include "IMeshBuilderModule.h"
 
 #define LOCTEXT_NAMESPACE "FbxStaticMeshImport"
 
@@ -576,14 +576,21 @@ bool UnFbx::FFbxImporter::BuildStaticMeshFromGeometry(FbxNode* Node, UStaticMesh
 				FVector VertexPosition = MeshDescription->GetVertex(VertexID).VertexPosition;
 				CornerVerticesIDs[CornerIndex] = VertexID;
 
-				FMeshVertexInstance NewVertexInstance;
-				NewVertexInstance.VertexID = VertexID;
-				FVertexInstanceID AddedVertexInstanceId = MeshDescription->VertexInstances().Add(NewVertexInstance);
+				FMeshVertexInstance TmpVertexInstance;
+				TmpVertexInstance.VertexID = VertexID;
+				FVertexInstanceID AddedVertexInstanceId = MeshDescription->VertexInstances().Add(TmpVertexInstance);
+				
+				//Make sure the Added vertex instance ID is matching the expected vertex instance ID
+				check(AddedVertexInstanceId == VertexInstanceID);
+				
+				FMeshVertexInstance& NewVertexInstance = MeshDescription->GetVertexInstance(AddedVertexInstanceId);
 				if (!MeshDescription->VertexInstances().IsValid(AddedVertexInstanceId) || AddedVertexInstanceId.GetValue() != VertexInstanceIndex)
 				{
 					AddTokenizedErrorMessage(FTokenizedMessage::Create(EMessageSeverity::Error, FText::Format(LOCTEXT("Error_CannotCreateVertexInstance", "Cannot create valid vertex instance for mesh '{0}'"), FText::FromString(Mesh->GetName()))), FFbxErrors::StaticMesh_BuildError);
 					return false;
 				}
+				FMeshVertex& MeshVertex = MeshDescription->GetVertex(VertexID);
+				MeshVertex.VertexInstanceIDs.Add(AddedVertexInstanceId);
 
 				//UVs attributes
 				for (int32 UVLayerIndex = 0; UVLayerIndex < FBXUVs.UniqueUVCount; UVLayerIndex++)
@@ -599,7 +606,8 @@ bool UnFbx::FFbxImporter::BuildStaticMeshFromGeometry(FbxNode* Node, UStaticMesh
 						FinalUVVector.X = static_cast<float>(UVVector[0]);
 						FinalUVVector.Y = 1.f - static_cast<float>(UVVector[1]);   //flip the Y of UVs for DirectX
 					}
-					MeshDescription->VertexInstanceAttributes().SetAttribute<FVector2D>(VertexInstanceID, UEditableMeshAttribute::VertexTextureCoordinate(), UVLayerIndex, FinalUVVector);
+					MeshDescription->VertexInstanceAttributes().SetAttribute<FVector2D>(AddedVertexInstanceId, UEditableMeshAttribute::VertexTextureCoordinate(), UVLayerIndex, FinalUVVector);
+					NewVertexInstance.VertexUVs.Add(FinalUVVector);
 				}
 
 				//Color attribute
@@ -621,7 +629,8 @@ bool UnFbx::FFbxImporter::BuildStaticMeshFromGeometry(FbxNode* Node, UStaticMesh
 							uint8(255.f*VertexColor.mBlue),
 							uint8(255.f*VertexColor.mAlpha)
 						);
-						MeshDescription->VertexInstanceAttributes().SetAttribute<FVector4>(VertexInstanceID, UEditableMeshAttribute::VertexColor(), 0, FVector4(FLinearColor(VertexInstanceColor)));
+						MeshDescription->VertexInstanceAttributes().SetAttribute<FVector4>(AddedVertexInstanceId, UEditableMeshAttribute::VertexColor(), 0, FVector4(FLinearColor(VertexInstanceColor)));
+						NewVertexInstance.Color = FLinearColor(VertexInstanceColor);
 					}
 				}
 				else if (VertexColorImportOption == EVertexColorImportOption::Ignore)
@@ -632,14 +641,16 @@ bool UnFbx::FFbxImporter::BuildStaticMeshFromGeometry(FbxNode* Node, UStaticMesh
 					if (PaintedColor)
 					{
 						// A matching color for this vertex was found
-						MeshDescription->VertexInstanceAttributes().SetAttribute<FVector4>(VertexInstanceID, UEditableMeshAttribute::VertexColor(), 0, FVector4(FLinearColor(*PaintedColor)));
+						MeshDescription->VertexInstanceAttributes().SetAttribute<FVector4>(AddedVertexInstanceId, UEditableMeshAttribute::VertexColor(), 0, FVector4(FLinearColor(*PaintedColor)));
+						NewVertexInstance.Color = FLinearColor(*PaintedColor);
 					}
 				}
 				else
 				{
 					// set the triangle's vertex color to a constant override
 					check(VertexColorImportOption == EVertexColorImportOption::Override);
-					MeshDescription->VertexInstanceAttributes().SetAttribute<FVector4>(VertexInstanceID, UEditableMeshAttribute::VertexColor(), 0, FVector4(FLinearColor(VertexOverrideColor)));
+					MeshDescription->VertexInstanceAttributes().SetAttribute<FVector4>(AddedVertexInstanceId, UEditableMeshAttribute::VertexColor(), 0, FVector4(FLinearColor(VertexOverrideColor)));
+					NewVertexInstance.Color = FLinearColor(VertexOverrideColor);
 				}
 
 				if (LayerElementNormal)
@@ -654,7 +665,8 @@ bool UnFbx::FFbxImporter::BuildStaticMeshFromGeometry(FbxNode* Node, UStaticMesh
 					FbxVector4 TempValue = LayerElementNormal->GetDirectArray().GetAt(NormalValueIndex);
 					TempValue = TotalMatrixForNormal.MultT(TempValue);
 					FVector TangentZ = Converter.ConvertDir(TempValue);
-					MeshDescription->VertexInstanceAttributes().SetAttribute<FVector>(CornerInstanceIDs[CornerIndex], UEditableMeshAttribute::VertexNormal(), 0, TangentZ.GetSafeNormal());
+					MeshDescription->VertexInstanceAttributes().SetAttribute<FVector>(AddedVertexInstanceId, UEditableMeshAttribute::VertexNormal(), 0, TangentZ.GetSafeNormal());
+					NewVertexInstance.Normal = TangentZ.GetSafeNormal();
 					//tangents and binormals share the same reference, mapping mode and index array
 					if (bHasNTBInformation)
 					{
@@ -666,7 +678,8 @@ bool UnFbx::FFbxImporter::BuildStaticMeshFromGeometry(FbxNode* Node, UStaticMesh
 						TempValue = LayerElementTangent->GetDirectArray().GetAt(TangentValueIndex);
 						TempValue = TotalMatrixForNormal.MultT(TempValue);
 						FVector TangentX = Converter.ConvertDir(TempValue);
-						MeshDescription->VertexInstanceAttributes().SetAttribute<FVector>(CornerInstanceIDs[CornerIndex], UEditableMeshAttribute::VertexTangent(), 0, TangentX.GetSafeNormal());
+						MeshDescription->VertexInstanceAttributes().SetAttribute<FVector>(AddedVertexInstanceId, UEditableMeshAttribute::VertexTangent(), 0, TangentX.GetSafeNormal());
+						NewVertexInstance.Tangent = TangentX.GetSafeNormal();
 
 						int BinormalMapIndex = (BinormalMappingMode == FbxLayerElement::eByControlPoint) ?
 							ControlPointIndex : TriangleCornerIndex;
@@ -677,7 +690,8 @@ bool UnFbx::FFbxImporter::BuildStaticMeshFromGeometry(FbxNode* Node, UStaticMesh
 						TempValue = TotalMatrixForNormal.MultT(TempValue);
 						FVector TangentY = -Converter.ConvertDir(TempValue);
 						FVector BiNormal = TangentY.GetSafeNormal();
-						MeshDescription->VertexInstanceAttributes().SetAttribute<float>(CornerInstanceIDs[CornerIndex], UEditableMeshAttribute::VertexBinormalSign(), 0, GetBasisDeterminantSign(TangentX.GetSafeNormal(), BiNormal, TangentZ.GetSafeNormal()) );
+						MeshDescription->VertexInstanceAttributes().SetAttribute<float>(AddedVertexInstanceId, UEditableMeshAttribute::VertexBinormalSign(), 0, GetBasisDeterminantSign(TangentX.GetSafeNormal(), BiNormal, TangentZ.GetSafeNormal()) );
+						NewVertexInstance.BinormalSign = GetBasisDeterminantSign(TangentX.GetSafeNormal(), BiNormal, TangentZ.GetSafeNormal());
 					}
 				}
 			}
@@ -759,16 +773,23 @@ bool UnFbx::FFbxImporter::BuildStaticMeshFromGeometry(FbxNode* Node, UStaticMesh
 			FPolygonGroupID PolygonGroupID = PolygonGroupMapping[RealMaterialIndex];
 
 			// Insert a polygon into the mesh
-			FMeshPolygon NewPolygon;
-			NewPolygon.PolygonGroupID = PolygonGroupID;
-
+			FMeshPolygon TmpPolygon;
+			TmpPolygon.PolygonGroupID = PolygonGroupID;
+			const FPolygonID NewPolygonID = MeshDescription->Polygons().Add(TmpPolygon);
+			FMeshPolygonGroup& MeshPolygonGroup = MeshDescription->GetPolygonGroup(PolygonGroupID);
+			MeshPolygonGroup.Polygons.Add(NewPolygonID);
+			FMeshPolygon& NewPolygon = MeshDescription->GetPolygon(NewPolygonID);
+			FMeshTriangle NewTriangle;
 			for (int32 TriangleVertexIndex = 0; TriangleVertexIndex < 3; ++TriangleVertexIndex)
 			{
 				const FVertexInstanceID &VertexInstanceID = CornerInstanceIDs[TriangleVertexIndex];
+				NewTriangle.SetVertexInstanceID(TriangleVertexIndex, VertexInstanceID);
+				FMeshVertexInstance& MeshVertexInstance = MeshDescription->GetVertexInstance(VertexInstanceID);
+				MeshVertexInstance.ConnectedPolygons.AddUnique(NewPolygonID);
 				const FVertexID &VertexID = CornerVerticesIDs[TriangleVertexIndex];
 				NewPolygon.PerimeterContour.VertexInstanceIDs.Add(VertexInstanceID);
 			}
-			const FPolygonID NewPolygonID = MeshDescription->Polygons().Add(NewPolygon);
+			NewPolygon.Triangles.Add(NewTriangle);
 
 			// Set edges hardness
 			{
@@ -813,6 +834,11 @@ bool UnFbx::FFbxImporter::BuildStaticMeshFromGeometry(FbxNode* Node, UStaticMesh
 					FMeshEdge& ExistingEdge = MeshDescription->GetEdge(MatchEdgeId);
 					ExistingEdge.ConnectedPolygons.Add(NewPolygonID);
 
+					FMeshVertex& MeshVertex0 = MeshDescription->GetVertex(EdgeVertexIDs[0]);
+					MeshVertex0.ConnectedEdgeIDs.Add(MatchEdgeId);
+					FMeshVertex& MeshVertex1 = MeshDescription->GetVertex(EdgeVertexIDs[1]);
+					MeshVertex1.ConnectedEdgeIDs.Add(MatchEdgeId);
+
 					if (!ExistingEdge.bIsHardEdge && bSmoothingAvailable && SmoothingInfo)
 					{
 						if (SmoothingMappingMode == FbxLayerElement::eByEdge)
@@ -1551,6 +1577,7 @@ UStaticMesh* UnFbx::FFbxImporter::ImportStaticMeshAsSingle(UObject* InParent, TA
 			//Create private asset in the same package as the StaticMesh, and make sure reference are set to avoid GC
 			FString MeshDescriptionName = MeshName + TEXT("_MeshData");
 			StaticMesh->MeshDescription = NewObject<UMeshDescription>(StaticMesh, FName(*MeshDescriptionName), RF_NoFlags);
+			MeshDescription = StaticMesh->MeshDescription;
 		}
 	}
 
@@ -1829,6 +1856,92 @@ UStaticMesh* UnFbx::FFbxImporter::ImportStaticMeshAsSingle(UObject* InParent, TA
 			FRawMesh LocalRawMesh;
 			SrcModel.RawMeshBulkData->LoadRawMesh(LocalRawMesh);
 		}
+		else
+		{
+			// Compress the materials array by removing any duplicates.
+			bool bDoRemap = false;
+			TArray<int32> MaterialMap;
+			TArray<FFbxMaterial> UniqueMaterials;
+			for (int32 MaterialIndex = 0; MaterialIndex < MeshMaterials.Num(); ++MaterialIndex)
+			{
+				bool bUnique = true;
+				for (int32 OtherMaterialIndex = MaterialIndex - 1; OtherMaterialIndex >= 0; --OtherMaterialIndex)
+				{
+					if (MeshMaterials[MaterialIndex].FbxMaterial == MeshMaterials[OtherMaterialIndex].FbxMaterial &&
+						MeshMaterials[MaterialIndex].Material == MeshMaterials[OtherMaterialIndex].Material)
+					{
+						int32 UniqueIndex = MaterialMap[OtherMaterialIndex];
+						MaterialMap.Add(UniqueIndex);
+						bDoRemap = true;
+						bUnique = false;
+						break;
+					}
+				}
+				if (bUnique)
+				{
+					int32 UniqueIndex = UniqueMaterials.Add(MeshMaterials[MaterialIndex]);
+					MaterialMap.Add(UniqueIndex);
+				}
+				else
+				{
+					UE_LOG(LogFbx, Verbose, TEXT("  remap %d -> %d"), MaterialIndex, MaterialMap[MaterialIndex]);
+				}
+			}
+
+			if (UniqueMaterials.Num() > LARGE_MESH_MATERIAL_INDEX_THRESHOLD)
+			{
+				AddTokenizedErrorMessage(
+					FTokenizedMessage::Create(
+						EMessageSeverity::Warning,
+						FText::Format(LOCTEXT("Error_TooManyMaterials", "StaticMesh has a large number({0}) of materials and may render inefficently.  Consider breaking up the mesh into multiple Static Mesh Assets"),
+							FText::AsNumber(UniqueMaterials.Num())
+						)),
+					FFbxErrors::StaticMesh_TooManyMaterials);
+			}
+
+			if (bDoRemap)
+			{
+				int32 GroupNumber = MeshDescription->PolygonGroups().Num();
+				TMap<int32, int32> PolygonGroupIDMap;
+				for (const FPolygonGroupID& PolygonGroupID : MeshDescription->PolygonGroups().GetElementIDs())
+				{
+					int32 MaterialIndex = PolygonGroupID.GetValue();
+					PolygonGroupIDMap.Add(MaterialIndex, MaterialMap[MaterialIndex]);
+				}
+
+				for (const FPolygonID& PolygonID : MeshDescription->Polygons().GetElementIDs())
+				{
+					FMeshPolygon MeshPolygon = MeshDescription->GetPolygon(PolygonID);
+					int32 MaterialIndex = MeshPolygon.PolygonGroupID.GetValue();
+					MeshPolygon.PolygonGroupID = FPolygonGroupID(PolygonGroupIDMap[MeshPolygon.PolygonGroupID.GetValue()]);
+				}
+				FPolygonGroupArray ReorderPolygonGroups;
+				for (int32 GroupIndex = 0; GroupIndex < GroupNumber; ++GroupIndex)
+				{
+					for (auto Kvp : PolygonGroupIDMap)
+					{
+						if (Kvp.Value == GroupIndex)
+						{
+							const FMeshPolygonGroup& MeshPolygonGroup = MeshDescription->GetPolygonGroup(FPolygonGroupID(Kvp.Key));
+							ReorderPolygonGroups.Insert(FPolygonGroupID(Kvp.Value), MeshPolygonGroup);
+							break;
+						}
+					}
+				}
+				MeshDescription->PolygonGroups().Reset(GroupNumber);
+				for (const FPolygonGroupID PolygonGroupID : ReorderPolygonGroups.GetElementIDs())
+				{
+					MeshDescription->PolygonGroups().Insert(PolygonGroupID, ReorderPolygonGroups[PolygonGroupID]);
+				}
+			}
+
+			for (const FPolygonGroupID& PolygonGroupID : MeshDescription->PolygonGroups().GetElementIDs())
+			{
+				int32 MaterialIndex = PolygonGroupID.GetValue();
+				const FMeshPolygonGroup& MeshPolygonGroup = MeshDescription->GetPolygonGroup(PolygonGroupID);
+				int32 AddIndex = StaticMesh->StaticMaterials.Add(FStaticMaterial(Cast<UMaterialInterface>(MeshPolygonGroup.MaterialAsset.TryLoad()), MeshPolygonGroup.MaterialSlotName, MeshPolygonGroup.ImportedMaterialSlotName));
+			}
+		}
 
 		// Setup default LOD settings based on the selected LOD group.
 		if (LODIndex == 0)
@@ -1980,7 +2093,8 @@ void UnFbx::FFbxImporter::PostImportStaticMesh(UStaticMesh* StaticMesh, TArray<F
 				
 				if (ImportOptions->bImportEditableMesh)
 				{
-					//TODO: Build the Mesh Description
+					IMeshBuilderModule& MeshBuilder = FModuleManager::Get().LoadModuleChecked<IMeshBuilderModule>(TEXT("MeshBuilder"));
+					MeshBuilder.BuildMesh(StaticMesh);
 				}
 				else
 				{
@@ -2024,7 +2138,8 @@ void UnFbx::FFbxImporter::PostImportStaticMesh(UStaticMesh* StaticMesh, TArray<F
 	}
 	if (ImportOptions->bImportEditableMesh)
 	{
-		//TODO: Build the mesh description
+		IMeshBuilderModule& MeshBuilder = FModuleManager::Get().LoadModuleChecked<IMeshBuilderModule>(TEXT("MeshBuilder"));
+		MeshBuilder.BuildMesh(StaticMesh);
 	}
 	else
 	{
diff --git a/Engine/Source/Editor/UnrealEd/UnrealEd.Build.cs b/Engine/Source/Editor/UnrealEd/UnrealEd.Build.cs
--- a/Engine/Source/Editor/UnrealEd/UnrealEd.Build.cs
+++ b/Engine/Source/Editor/UnrealEd/UnrealEd.Build.cs
@@ -96,6 +96,7 @@ public class UnrealEd : ModuleRules
 				"NetworkFileSystem",
                 "UMG",
                 "MeshDescription",
+                "MeshBuilder",
             }
 		);
 
diff --git a/Engine/Source/Runtime/Engine/Private/StaticMeshBuild.cpp b/Engine/Source/Runtime/Engine/Private/StaticMeshBuild.cpp
--- a/Engine/Source/Runtime/Engine/Private/StaticMeshBuild.cpp
+++ b/Engine/Source/Runtime/Engine/Private/StaticMeshBuild.cpp
@@ -60,7 +60,7 @@ void UStaticMesh::Build(bool bSilent, TArray<FText>* OutErrors)
 		return;
 
 	// If we're controlled by an editable mesh do not build. The editable mesh will build us
-	if (EditableMesh)
+	if (EditableMesh || MeshDescription)
 	{
 		return;
 	}
diff --git a/Engine/Source/Runtime/MeshBuilder/MeshBuilder.Build.cs b/Engine/Source/Runtime/MeshBuilder/MeshBuilder.Build.cs
new file mode 100644
--- /dev/null
+++ b/Engine/Source/Runtime/MeshBuilder/MeshBuilder.Build.cs
@@ -0,0 +1,24 @@
+// Copyright 1998-2017 Epic Games, Inc. All Rights Reserved.
+
+namespace UnrealBuildTool.Rules
+{
+	public class MeshBuilder : ModuleRules
+	{
+		public MeshBuilder(ReadOnlyTargetRules Target) : base(Target)
+		{
+            PrivateIncludePaths.Add("Runtime/MeshBuilder/Private");
+            PublicIncludePaths.Add("Runtime/MeshBuilder/Public");
+
+            PrivateDependencyModuleNames.AddRange(
+				new string[]
+				{
+					"Core",
+					"CoreUObject",
+                    "Engine",
+                    "RenderCore",
+                    "MeshDescription",
+                }
+			);
+		}
+	}
+}
diff --git a/Engine/Source/Runtime/MeshBuilder/Private/MeshBuilder.cpp b/Engine/Source/Runtime/MeshBuilder/Private/MeshBuilder.cpp
new file mode 100644
--- /dev/null
+++ b/Engine/Source/Runtime/MeshBuilder/Private/MeshBuilder.cpp
@@ -0,0 +1,8 @@
+// Copyright 1998-2017 Epic Games, Inc. All Rights Reserved.
+
+#include "MeshBuilder.h"
+
+FMeshBuilder::FMeshBuilder()
+{
+
+}
\ No newline at end of file
diff --git a/Engine/Source/Runtime/MeshBuilder/Private/MeshBuilderModule.cpp b/Engine/Source/Runtime/MeshBuilder/Private/MeshBuilderModule.cpp
new file mode 100644
--- /dev/null
+++ b/Engine/Source/Runtime/MeshBuilder/Private/MeshBuilderModule.cpp
@@ -0,0 +1,44 @@
+// Copyright 1998-2017 Epic Games, Inc. All Rights Reserved.
+
+#include "IMeshBuilderModule.h"
+#include "Modules/ModuleManager.h"
+
+#include "StaticMeshBuilder.h"
+#include "Engine/StaticMesh.h"
+
+class FMeshBuilderModule : public IMeshBuilderModule
+{
+public:
+
+	FMeshBuilderModule()
+	{
+	}
+
+	virtual void StartupModule() override
+	{
+		// Register any modular features here
+	}
+
+	virtual void ShutdownModule() override
+	{
+		// Unregister any modular features here
+	}
+
+	virtual bool BuildMesh(class UObject* Mesh) override;
+
+private:
+
+};
+
+IMPLEMENT_MODULE(FMeshBuilderModule, MeshBuilder );
+
+bool FMeshBuilderModule::BuildMesh(class UObject* Mesh)
+{
+	UStaticMesh* StaticMesh = Cast<UStaticMesh>(Mesh);
+	if (StaticMesh != nullptr)
+	{
+		//Call the static mesh builder
+		return FStaticMeshBuilder::Get().Build(StaticMesh);
+	}
+	return false;
+}
diff --git a/Engine/Source/Runtime/MeshBuilder/Private/StaticMeshBuilder.cpp b/Engine/Source/Runtime/MeshBuilder/Private/StaticMeshBuilder.cpp
new file mode 100644
--- /dev/null
+++ b/Engine/Source/Runtime/MeshBuilder/Private/StaticMeshBuilder.cpp
@@ -0,0 +1,369 @@
+// Copyright 1998-2017 Epic Games, Inc. All Rights Reserved.
+
+#include "StaticMeshBuilder.h"
+#include "Engine/StaticMesh.h"
+#include "StaticMeshResources.h"
+#include "PhysicsEngine/BodySetup.h"
+#include "MeshDescription.h"
+
+const FTriangleID FTriangleID::Invalid(TNumericLimits<uint32>::Max());
+
+FStaticMeshBuilder::FStaticMeshBuilder()
+{
+
+}
+
+bool FStaticMeshBuilder::Build(UStaticMesh* StaticMesh)
+{
+	OnBuildRenderMeshStart(StaticMesh, false);
+
+	FStaticMeshRenderData& StaticMeshRenderData = *StaticMesh->RenderData;
+	//TODO: support LOD in the mesh description
+	FStaticMeshLODResources& StaticMeshLOD = StaticMeshRenderData.LODResources[0];
+
+	// Build new vertex buffers
+	static TArray< FStaticMeshBuildVertex > StaticMeshBuildVertices;
+	StaticMeshBuildVertices.Reset();
+
+	static TArray< uint32 > IndexBuffer;
+	IndexBuffer.Reset();
+
+	const FPolygonGroupArray& PolygonGroups = StaticMesh->MeshDescription->PolygonGroups();
+
+	StaticMeshLOD.Sections.Empty(PolygonGroups.Num());
+
+	bool bHasColor = false;
+
+	const FVertexArray& Vertices = StaticMesh->MeshDescription->Vertices();
+	const FVertexInstanceArray& VertexInstances = StaticMesh->MeshDescription->VertexInstances();
+
+	// set up vertex buffer elements
+	StaticMeshBuildVertices.SetNum(VertexInstances.GetArraySize());
+
+	for (const FVertexInstanceID VertexInstanceID : VertexInstances.GetElementIDs())
+	{
+		const FMeshVertexInstance& VertexInstance = VertexInstances[VertexInstanceID];
+
+		if (VertexInstance.Color != FColor::White)
+		{
+			bHasColor = true;
+		}
+
+		FStaticMeshBuildVertex& StaticMeshVertex = StaticMeshBuildVertices[VertexInstanceID.GetValue()];
+
+		StaticMeshVertex.Position = Vertices[VertexInstance.VertexID].VertexPosition;
+		StaticMeshVertex.TangentX = VertexInstance.Tangent;
+		StaticMeshVertex.TangentY = FVector::CrossProduct(VertexInstance.Normal, VertexInstance.Tangent).GetSafeNormal() * VertexInstance.BinormalSign;
+		StaticMeshVertex.TangentZ = VertexInstance.Normal;
+		StaticMeshVertex.Color = VertexInstance.Color.ToFColor(true);
+		for (int32 UVIndex = 0; UVIndex < VertexInstance.VertexUVs.Num(); ++UVIndex)
+		{
+			StaticMeshVertex.UVs[UVIndex] = VertexInstance.VertexUVs[UVIndex];
+		}
+	}
+
+	
+	/** All of the polygon groups in this mesh */
+	TMeshElementArray<FRenderingPolygonGroup, FPolygonGroupID> RenderingPolygonGroups;
+
+	for (const FPolygonID& PolygonID : StaticMesh->MeshDescription->Polygons().GetElementIDs())
+	{
+		FMeshPolygon MeshPolygon = StaticMesh->MeshDescription->GetPolygon(PolygonID);
+		FPolygonGroupID PolygonGroupID = MeshPolygon.PolygonGroupID;
+
+		if (!RenderingPolygonGroups.IsValid(PolygonGroupID))
+		{
+			// Create a rendering polygon group for holding the triangulated data and references to the static mesh rendering section.
+			// This is indexed by the same FPolygonGroupID as the PolygonGroups.
+			RenderingPolygonGroups.Insert(PolygonGroupID);
+		}
+		FRenderingPolygonGroup& RenderingPolygonGroup = RenderingPolygonGroups[PolygonGroupID];
+		for (const FMeshTriangle& MeshTriangle : MeshPolygon.Triangles)
+		{
+			RenderingPolygonGroup.Triangles.Add(MeshTriangle);
+		}
+		RenderingPolygonGroup.MaxTriangles = RenderingPolygonGroup.Triangles.Num();
+	}
+
+
+	// Set up index buffer
+	for (const FPolygonGroupID PolygonGroupID : PolygonGroups.GetElementIDs())
+	{
+		const FMeshPolygonGroup& PolygonGroup = PolygonGroups[PolygonGroupID];
+		
+		if (!RenderingPolygonGroups.IsValid(PolygonGroupID))
+		{
+			//There is no polygon using this polygon group
+			continue;
+		}
+		FRenderingPolygonGroup& RenderingPolygonGroup = RenderingPolygonGroups[PolygonGroupID];
+		RenderingPolygonGroup.RenderingSectionIndex = StaticMeshLOD.Sections.Num();
+
+		// Create new rendering section
+		StaticMeshLOD.Sections.Add(FStaticMeshSection());
+		FStaticMeshSection& StaticMeshSection = StaticMeshLOD.Sections.Last();
+
+		StaticMeshSection.FirstIndex = IndexBuffer.Num();
+
+		//TODO: Get the number of triangles
+		StaticMeshSection.NumTriangles = RenderingPolygonGroup.Triangles.GetArraySize();
+		check(RenderingPolygonGroup.Triangles.GetArraySize() <= RenderingPolygonGroup.MaxTriangles);
+
+		const int32 MaterialIndex = StaticMesh->GetMaterialIndex(PolygonGroup.MaterialSlotName);
+		check(MaterialIndex != INDEX_NONE);
+		check(StaticMesh->StaticMaterials[MaterialIndex].MaterialInterface->GetPathName() == PolygonGroup.MaterialAsset.ToString());
+		StaticMeshSection.MaterialIndex = MaterialIndex;
+		StaticMeshSection.bEnableCollision = PolygonGroup.bEnableCollision;
+		StaticMeshSection.bCastShadow = PolygonGroup.bCastShadow;
+
+
+		if (RenderingPolygonGroup.Triangles.Num() > 0)
+		{
+			IndexBuffer.Reserve(IndexBuffer.Num() + RenderingPolygonGroup.Triangles.GetArraySize() * 3);
+			uint32 MinIndex = TNumericLimits< uint32 >::Max();
+			uint32 MaxIndex = TNumericLimits< uint32 >::Min();
+
+			// Find the first valid vertex instance index, so that we have a value we can use for our degenerates
+			check(RenderingPolygonGroup.Triangles.Num() > 0);
+			const FVertexInstanceID FirstValidRenderingID = RenderingPolygonGroup.Triangles[RenderingPolygonGroup.Triangles.GetFirstValidID()].GetVertexInstanceID(0);
+
+			for (int32 TriangleIndex = 0; TriangleIndex < RenderingPolygonGroup.Triangles.GetArraySize(); ++TriangleIndex)
+			{
+				const FTriangleID TriangleID(TriangleIndex);
+				if (RenderingPolygonGroup.Triangles.IsValid(TriangleID))
+				{
+					const FMeshTriangle& Triangle = RenderingPolygonGroup.Triangles[TriangleID];
+					for (int32 TriVert = 0; TriVert < 3; ++TriVert)
+					{
+						const uint32 RenderingVertexIndex = Triangle.GetVertexInstanceID(TriVert).GetValue();
+						IndexBuffer.Add(RenderingVertexIndex);
+						MinIndex = FMath::Min(MinIndex, RenderingVertexIndex);
+						MaxIndex = FMath::Max(MaxIndex, RenderingVertexIndex);
+					}
+				}
+				else
+				{
+					IndexBuffer.Add(FirstValidRenderingID.GetValue());
+					IndexBuffer.Add(FirstValidRenderingID.GetValue());
+					IndexBuffer.Add(FirstValidRenderingID.GetValue());
+				}
+			}
+
+			StaticMeshSection.MinVertexIndex = MinIndex;
+			StaticMeshSection.MaxVertexIndex = MaxIndex;
+
+			// Add any index buffer padding.
+			// This can be necessary if we have just loaded an editable mesh which had a MaxTriangles count in the editable mesh section
+			// greater than the sparse array max size (i.e. an extra gap had been reserved for tris).
+			const int32 IndexBufferPadding = RenderingPolygonGroup.MaxTriangles - RenderingPolygonGroup.Triangles.GetArraySize();
+			if (IndexBufferPadding > 0)
+			{
+				IndexBuffer.AddZeroed(IndexBufferPadding * 3);
+			}
+		}
+		else
+		{
+			// No triangles in this section
+			StaticMeshSection.MinVertexIndex = 0;
+			StaticMeshSection.MaxVertexIndex = 0;
+		}
+	}
+
+	// Figure out which index buffer stride we need
+	bool bNeeds32BitIndices = false;
+	for (const FStaticMeshSection& StaticMeshSection : StaticMeshLOD.Sections)
+	{
+		if (StaticMeshSection.MaxVertexIndex > TNumericLimits<uint16>::Max())
+		{
+			bNeeds32BitIndices = true;
+		}
+	}
+	const EIndexBufferStride::Type IndexBufferStride = bNeeds32BitIndices ? EIndexBufferStride::Force32Bit : EIndexBufferStride::Force16Bit;
+
+	StaticMeshLOD.PositionVertexBuffer.Init(StaticMeshBuildVertices);
+	StaticMeshLOD.VertexBuffer.Init(StaticMeshBuildVertices, 1); //TODO: use the real texture coordinnate count
+
+	if (bHasColor)
+	{
+		StaticMeshLOD.ColorVertexBuffer.Init(StaticMeshBuildVertices);
+	}
+	else
+	{
+		StaticMeshLOD.ColorVertexBuffer.InitFromSingleColor(FColor::White, StaticMeshBuildVertices.Num());
+	}
+
+	StaticMeshLOD.IndexBuffer.SetIndices(IndexBuffer, IndexBufferStride);
+
+	// @todo mesheditor: support the other index buffer types
+	StaticMeshLOD.ReversedIndexBuffer.SetIndices(TArray< uint32 >(), IndexBufferStride);
+	StaticMeshLOD.DepthOnlyIndexBuffer.SetIndices(TArray< uint32 >(), IndexBufferStride);
+	StaticMeshLOD.ReversedDepthOnlyIndexBuffer.SetIndices(TArray< uint32 >(), IndexBufferStride);
+	StaticMeshLOD.WireframeIndexBuffer.SetIndices(TArray< uint32 >(), IndexBufferStride);
+	StaticMeshLOD.AdjacencyIndexBuffer.SetIndices(TArray< uint32 >(), IndexBufferStride);
+
+	StaticMeshLOD.bHasAdjacencyInfo = false;
+	StaticMeshLOD.bHasDepthOnlyIndices = false;
+	StaticMeshLOD.bHasReversedIndices = false;
+	StaticMeshLOD.bHasReversedDepthOnlyIndices = false;
+	StaticMeshLOD.DepthOnlyNumTriangles = 0;
+
+	OnBuildRenderMeshFinish(StaticMesh, true);
+	return true;
+}
+
+
+void FStaticMeshBuilder::OnBuildRenderMeshStart(UStaticMesh* StaticMesh, const bool bInvalidateLighting)
+{
+	// We may already have a lock on the rendering resources, if it wasn't released the last time we called EndModification()
+	// on the mesh.  This is only the case when rolling back preview changes for a mesh, because we're guaranteed to apply another
+	// modification to the same mesh in the very same frame.  So we can avoid having to update the GPU resources twice in one frame.
+	if (!RecreateRenderStateContext.IsValid())
+	{
+		// We're changing the mesh itself, so ALL static mesh components in the scene will need
+		// to be unregistered for this (and reregistered afterwards.)
+		const bool bRefreshBounds = true;
+		RecreateRenderStateContext = MakeShareable(new FStaticMeshComponentRecreateRenderStateContext(StaticMesh, bInvalidateLighting, bRefreshBounds));
+
+		// Release the static mesh's resources.
+		StaticMesh->ReleaseResources();
+
+		// Flush the resource release commands to the rendering thread to ensure that the build doesn't occur while a resource is still
+		// allocated, and potentially accessing the UStaticMesh.
+		StaticMesh->ReleaseResourcesFence.Wait();
+	}
+}
+
+void FStaticMeshBuilder::OnBuildRenderMeshFinish(UStaticMesh* StaticMesh, const bool bRebuildBoundsAndCollision)
+{
+	
+
+	if (bRebuildBoundsAndCollision)
+	{
+		UpdateBounds(StaticMesh);
+		UpdateCollision(StaticMesh);
+	}
+
+	StaticMesh->InitResources();
+
+	// NOTE: This can call InvalidateLightingCache() on all components using this mesh, causing Modify() to be 
+	// called on those components!  Just something to be aware of when EndModification() is called within
+	// an undo transaction.
+	RecreateRenderStateContext.Reset();
+}
+
+bool FStaticMeshBuilder::IsOrphanedVertex(UStaticMesh* StaticMesh, const FVertexID VertexID) const
+{
+	const FVertexInstanceArray& VertexInstances = StaticMesh->MeshDescription->VertexInstances();
+	for (const FVertexInstanceID VertexInstanceID : StaticMesh->MeshDescription->GetVertex(VertexID).VertexInstanceIDs)
+	{
+		if (VertexInstances[VertexInstanceID].ConnectedPolygons.Num() > 0)
+		{
+			return false;
+		}
+	}
+
+	return true;
+}
+
+void FStaticMeshBuilder::UpdateBounds(UStaticMesh* StaticMesh)
+{
+	// Compute a new bounding box
+	// @todo mesheditor perf: During the final modification, only do this if the bounds may have changed (need hinting)
+	FBoxSphereBounds BoundingBoxAndSphere;
+
+	FBox BoundingBox;
+	BoundingBox.Init();
+
+	// Could improve performance here if necessary:
+	// 1) cache polygon IDs per vertex (in order to quickly reject orphans) and just iterate vertex array; or
+	// 2) cache bounding box per polygon
+	// There are other cases where having polygon adjacency information (1) might be useful, so it's maybe worth considering.
+
+	const FVertexArray& Vertices = StaticMesh->MeshDescription->Vertices();
+	const FVertexInstanceArray& VertexInstances = StaticMesh->MeshDescription->VertexInstances();
+
+	for (const FVertexID VertexID : Vertices.GetElementIDs())
+	{
+		if (!IsOrphanedVertex(StaticMesh, VertexID))
+		{
+			BoundingBox += Vertices[VertexID].VertexPosition;
+		}
+	}
+
+	BoundingBox.GetCenterAndExtents( /* Out */ BoundingBoxAndSphere.Origin, /* Out */ BoundingBoxAndSphere.BoxExtent);
+
+	// Calculate the bounding sphere, using the center of the bounding box as the origin.
+	BoundingBoxAndSphere.SphereRadius = 0.0f;
+
+	for (const FVertexID VertexID : Vertices.GetElementIDs())
+	{
+		if (!IsOrphanedVertex(StaticMesh, VertexID))
+		{
+			const FVector VertexPosition = Vertices[VertexID].VertexPosition;
+			BoundingBoxAndSphere.SphereRadius = FMath::Max((VertexPosition - BoundingBoxAndSphere.Origin).Size(), BoundingBoxAndSphere.SphereRadius);
+		}
+	}
+
+	FStaticMeshRenderData& StaticMeshRenderData = *StaticMesh->RenderData;
+	StaticMeshRenderData.Bounds = BoundingBoxAndSphere;
+	StaticMesh->CalculateExtendedBounds();
+}
+
+
+void FStaticMeshBuilder::UpdateCollision(UStaticMesh *StaticMesh)
+{
+	// @todo mesheditor collision: We're wiping the existing simplified collision and generating a simple bounding
+	// box collision, since that's the best we can do without impacting performance.  We always using visibility (complex)
+	// collision for traces while mesh editing (for hover/selection), so simplified collision isn't really important.
+	const bool bRecreateSimplifiedCollision = true;
+
+	if (StaticMesh->BodySetup == nullptr)
+	{
+		StaticMesh->CreateBodySetup();
+	}
+
+	UBodySetup* BodySetup = StaticMesh->BodySetup;
+
+	// NOTE: We don't bother calling Modify() on the BodySetup as EndModification() will rebuild this guy after every undo
+	// BodySetup->Modify();
+
+	if (bRecreateSimplifiedCollision)
+	{
+		if (BodySetup->AggGeom.GetElementCount() > 0)
+		{
+			BodySetup->RemoveSimpleCollision();
+		}
+	}
+
+	BodySetup->InvalidatePhysicsData();
+
+	if (bRecreateSimplifiedCollision)
+	{
+		const FBoxSphereBounds Bounds = StaticMesh->GetBounds();
+
+		FKBoxElem BoxElem;
+		BoxElem.Center = Bounds.Origin;
+		BoxElem.X = Bounds.BoxExtent.X * 2.0f;
+		BoxElem.Y = Bounds.BoxExtent.Y * 2.0f;
+		BoxElem.Z = Bounds.BoxExtent.Z * 2.0f;
+		BodySetup->AggGeom.BoxElems.Add(BoxElem);
+	}
+
+	// Update all static mesh components that are using this mesh
+	// @todo mesheditor perf: This is a pretty heavy operation, and overlaps with what we're already doing in RecreateRenderStateContext
+	// a little bit.  Ideally we do everything in a single pass.  Furthermore, if this could be updated lazily it would be faster.
+	{
+		for (FObjectIterator Iter(UStaticMeshComponent::StaticClass()); Iter; ++Iter)
+		{
+			UStaticMeshComponent* StaticMeshComponent = Cast<UStaticMeshComponent>(*Iter);
+			if (StaticMeshComponent->GetStaticMesh() == StaticMesh)
+			{
+				// it needs to recreate IF it already has been created
+				if (StaticMeshComponent->IsPhysicsStateCreated())
+				{
+					StaticMeshComponent->RecreatePhysicsState();
+				}
+			}
+		}
+	}
+}
diff --git a/Engine/Source/Runtime/MeshBuilder/Public/IMeshBuilderModule.h b/Engine/Source/Runtime/MeshBuilder/Public/IMeshBuilderModule.h
new file mode 100644
--- /dev/null
+++ b/Engine/Source/Runtime/MeshBuilder/Public/IMeshBuilderModule.h
@@ -0,0 +1,34 @@
+// Copyright 1998-2017 Epic Games, Inc. All Rights Reserved.
+
+#pragma once
+
+#include "CoreMinimal.h"
+#include "Modules/ModuleInterface.h"
+
+class IMeshBuilderModule : public IModuleInterface
+{
+public:
+
+	/**
+	 * Singleton-like access to this module's interface.  This is just for convenience!
+	 * Beware of calling this during the shutdown phase, though.  Your module might have been unloaded already.
+	 *
+	 * @return Returns singleton instance, loading the module on demand if needed
+	 */
+	static inline IMeshBuilderModule & Get()
+	{
+		return FModuleManager::LoadModuleChecked<IMeshBuilderModule>( "MeshBuilder" );
+	}
+
+	/**
+	 * Checks to see if this module is loaded and ready.  It is only valid to call Get() if IsAvailable() returns true.
+	 *
+	 * @return True if the module is loaded and ready to use
+	 */
+	static inline bool IsAvailable()
+	{
+		return FModuleManager::Get().IsModuleLoaded( "MeshBuilder" );
+	}
+
+	virtual bool BuildMesh(class UObject* Mesh) { return false; }
+};
diff --git a/Engine/Source/Runtime/MeshBuilder/Public/MeshBuilder.h b/Engine/Source/Runtime/MeshBuilder/Public/MeshBuilder.h
new file mode 100644
--- /dev/null
+++ b/Engine/Source/Runtime/MeshBuilder/Public/MeshBuilder.h
@@ -0,0 +1,25 @@
+// Copyright 1998-2017 Epic Games, Inc. All Rights Reserved.
+
+#pragma once
+
+#include "CoreMinimal.h"
+#include "Materials/MaterialInterface.h"
+
+/**
+ * Abstract class which is the base class of all builder.
+ * All share code to build some render data should be found inside this class
+ */
+class MESHBUILDER_API FMeshBuilder
+{
+public:
+	FMeshBuilder();
+	
+	/**
+	 * Build function should be override and is the starting point for all mesh builders
+	 */
+	virtual bool Build(class UStaticMesh* StaticMesh) =0;
+
+private:
+
+};
+
diff --git a/Engine/Source/Runtime/MeshBuilder/Public/StaticMeshBuilder.h b/Engine/Source/Runtime/MeshBuilder/Public/StaticMeshBuilder.h
new file mode 100644
--- /dev/null
+++ b/Engine/Source/Runtime/MeshBuilder/Public/StaticMeshBuilder.h
@@ -0,0 +1,99 @@
+// Copyright 1998-2017 Epic Games, Inc. All Rights Reserved.
+
+#pragma once
+
+#include "CoreMinimal.h"
+#include "MeshBuilder.h"
+#include "MeshDescription.h"
+
+
+class MESHBUILDER_API FStaticMeshBuilder : public FMeshBuilder
+{
+public:
+	static FStaticMeshBuilder& Get()
+	{
+		static FStaticMeshBuilder Instance;
+		return Instance;
+	}
+
+	virtual bool Build(class UStaticMesh* StaticMesh) override;
+	virtual ~FStaticMeshBuilder() {}
+
+private:
+	FStaticMeshBuilder();
+
+	void OnBuildRenderMeshStart(class UStaticMesh* StaticMesh, const bool bInvalidateLighting);
+	void OnBuildRenderMeshFinish(class UStaticMesh* StaticMesh, const bool bRebuildBoundsAndCollision);
+	
+	bool IsOrphanedVertex(class UStaticMesh* StaticMesh, const FVertexID VertexID) const;
+	void UpdateBounds(class UStaticMesh* StaticMesh);
+	void UpdateCollision(UStaticMesh *StaticMesh);
+
+	/** Used to refresh all components in the scene that may be using a mesh we're editing */
+	TSharedPtr<class FStaticMeshComponentRecreateRenderStateContext> RecreateRenderStateContext;
+};
+
+struct FTriangleID : public FElementID
+{
+	FTriangleID()
+	{
+	}
+
+	explicit FTriangleID(const FElementID InitElementID)
+		: FElementID(InitElementID.GetValue())
+	{
+	}
+
+	explicit FTriangleID(const uint32 InitIDValue)
+		: FElementID(InitIDValue)
+	{
+	}
+
+	FORCEINLINE friend uint32 GetTypeHash(const FTriangleID& Other)
+	{
+		return GetTypeHash(Other.IDValue);
+	}
+
+	/** Invalid triangle ID */
+	static const FTriangleID Invalid;
+};
+
+
+struct FRenderingPolygon
+{
+
+	/** Which rendering polygon group the polygon is in */
+	FPolygonGroupID PolygonGroupID;
+
+	/** This is a list of indices of triangles in the FRenderingPolygonGroup::Triangles array.
+	We use this to maintain a record of which triangles in the section belong to this polygon. */
+	TArray<FTriangleID> TriangulatedPolygonTriangleIndices;
+};
+
+
+struct FRenderingPolygonGroup
+{
+
+	/** The rendering section index for this mesh section */
+	uint32 RenderingSectionIndex;
+
+	/** Maximum number of triangles which have been reserved in the index buffer */
+	int32 MaxTriangles;
+
+	/** Sparse array of triangles, that matches the triangles in the mesh index buffers.  Elements that
+	aren't allocated will be stored as degenerates in the mesh index buffer. */
+	TMeshElementArray<FMeshTriangle, FTriangleID> Triangles;
+
+
+	/** Converts from an index in our Triangles array to an index of a rendering triangle's first vertex in the rendering mesh's index buffer */
+	inline static FTriangleID RenderingTriangleFirstIndexToTriangleIndex(const FStaticMeshSection& RenderingSection, const uint32 RenderingTriangleFirstIndex)
+	{
+		return FTriangleID((RenderingTriangleFirstIndex - RenderingSection.FirstIndex) / 3);
+	}
+
+	/** Converts from an index of a rendering triangle's first vertex in the rendering mesh's index buffer to an index in our Triangles array */
+	inline static uint32 TriangleIndexToRenderingTriangleFirstIndex(const FStaticMeshSection& RenderingSection, const FTriangleID TriangleIndex)
+	{
+		return TriangleIndex.GetValue() * 3 + RenderingSection.FirstIndex;
+	}
+};
