Support LOD import using mesh description
Optimize the fbx importer when accessing the edge
Optimize the fbx import by finding the edge from a vertex instead of iterating all existing edges

diff --git a/Engine/Plugins/Runtime/EditableMesh/Source/EditableMesh/EditableStaticMeshAdapter.cpp b/Engine/Plugins/Runtime/EditableMesh/Source/EditableMesh/EditableStaticMeshAdapter.cpp
--- a/Engine/Plugins/Runtime/EditableMesh/Source/EditableMesh/EditableStaticMeshAdapter.cpp
+++ b/Engine/Plugins/Runtime/EditableMesh/Source/EditableMesh/EditableStaticMeshAdapter.cpp
@@ -614,7 +614,7 @@ void UEditableStaticMeshAdapter::OnRebuildRenderMesh( const UEditableMesh* Edita
 			StaticMeshSection.MinVertexIndex = FirstSectionVertexIndex;
 			StaticMeshSection.MaxVertexIndex = FirstSectionVertexIndex + SectionTriangleCount * 3;
 
-			const int32 MaterialIndex = StaticMesh->GetMaterialIndex( PolygonGroup.MaterialSlotName );
+			const int32 MaterialIndex = StaticMesh->GetMaterialIndexFromImportedMaterialSlotName( PolygonGroup.ImportedMaterialSlotName );
 			check( MaterialIndex != INDEX_NONE );
 			check( StaticMesh->StaticMaterials[ MaterialIndex ].MaterialInterface->GetPathName() == PolygonGroup.MaterialAsset.ToString() );
 
@@ -724,7 +724,7 @@ void UEditableStaticMeshAdapter::OnRebuildRenderMesh( const UEditableMesh* Edita
 			StaticMeshSection.NumTriangles = RenderingPolygonGroup.Triangles.GetArraySize();
 			check( RenderingPolygonGroup.Triangles.GetArraySize() <= RenderingPolygonGroup.MaxTriangles );
 
-			const int32 MaterialIndex = StaticMesh->GetMaterialIndex( PolygonGroup.MaterialSlotName );
+			const int32 MaterialIndex = StaticMesh->GetMaterialIndexFromImportedMaterialSlotName( PolygonGroup.ImportedMaterialSlotName );
 			check( MaterialIndex != INDEX_NONE );
 			check( StaticMesh->StaticMaterials[ MaterialIndex ].MaterialInterface->GetPathName() == PolygonGroup.MaterialAsset.ToString() );
 			StaticMeshSection.MaterialIndex = MaterialIndex;
diff --git a/Engine/Source/Editor/UnrealEd/Private/Fbx/FbxStaticMeshImport.cpp b/Engine/Source/Editor/UnrealEd/Private/Fbx/FbxStaticMeshImport.cpp
--- a/Engine/Source/Editor/UnrealEd/Private/Fbx/FbxStaticMeshImport.cpp
+++ b/Engine/Source/Editor/UnrealEd/Private/Fbx/FbxStaticMeshImport.cpp
@@ -265,11 +265,11 @@ bool UnFbx::FFbxImporter::BuildStaticMeshFromGeometry(FbxNode* Node, UStaticMesh
 	FbxMesh* Mesh = Node->GetMesh();
 	FStaticMeshSourceModel& SrcModel = StaticMesh->SourceModels[LODIndex];
 	
-	UMeshDescription *MeshDescription = StaticMesh->MeshDescription;
-	if (ImportOptions->bImportEditableMesh && MeshDescription == nullptr)
+	UMeshDescription *MeshDescription = StaticMesh->GetMeshDescription(LODIndex);
+	if (ImportOptions->bImportEditableMesh)
 	{
-		StaticMesh->MeshDescription = NewObject<UMeshDescription>(StaticMesh, NAME_None, RF_NoFlags);
-		MeshDescription = StaticMesh->MeshDescription;
+		//The mesh description should have been created before calling BuildStaticMeshFromGeometry
+		check(MeshDescription);
 	}
 	//remove the bad polygons before getting any data from mesh
 	Mesh->RemoveBadPolygons();
@@ -528,8 +528,8 @@ bool UnFbx::FFbxImporter::BuildStaticMeshFromGeometry(FbxNode* Node, UStaticMesh
 		{
 			StaticMesh->CacheDerivedData();
 		}
-		
-		if (StaticMesh->RenderData->LODResources.Num() == 0)
+		//Add missing LODs render data resources
+		while (StaticMesh->RenderData->LODResources.Num() <= LODIndex)
 		{
 			FStaticMeshLODResources *DummyStaticMeshLodResources = new FStaticMeshLODResources();
 			StaticMesh->RenderData->LODResources.Add(DummyStaticMeshLodResources);
@@ -559,6 +559,9 @@ bool UnFbx::FFbxImporter::BuildStaticMeshFromGeometry(FbxNode* Node, UStaticMesh
 				return false;
 			}
 		}
+
+		//Call this before all GetMeshEdgeIndexForPolygon call this is for optimization purpose.
+		Mesh->BeginGetMeshEdgeIndexForPolygon();
 		//Polygons
 		for (int32 TriangleIndex = 0; TriangleIndex < TriangleCount; TriangleIndex++)
 		{
@@ -772,7 +775,11 @@ bool UnFbx::FFbxImporter::BuildStaticMeshFromGeometry(FbxNode* Node, UStaticMesh
 			TmpPolygon.PolygonGroupID = PolygonGroupID;
 			const FPolygonID NewPolygonID = MeshDescription->Polygons().Add(TmpPolygon);
 			FMeshPolygonGroup& MeshPolygonGroup = MeshDescription->GetPolygonGroup(PolygonGroupID);
-			MeshPolygonGroup.Polygons.AddUnique(NewPolygonID);
+			
+			//We cannot do a add unique here since this is too slow, since we just create the polygons it should not be part
+			//of the polygonGroup already
+			MeshPolygonGroup.Polygons.Add(NewPolygonID);
+			
 			FMeshPolygon& NewPolygon = MeshDescription->GetPolygon(NewPolygonID);
 			FMeshTriangle NewTriangle;
 			for (int32 TriangleVertexIndex = 0; TriangleVertexIndex < 3; ++TriangleVertexIndex)
@@ -852,6 +859,8 @@ bool UnFbx::FFbxImporter::BuildStaticMeshFromGeometry(FbxNode* Node, UStaticMesh
 				}
 			}
 		}
+		//Call this after all GetMeshEdgeIndexForPolygon call this is for optimization purpose.
+		Mesh->EndGetMeshEdgeIndexForPolygon();
 		// needed?
 		FBXUVs.Cleanup();
 	}
@@ -1557,13 +1566,13 @@ UStaticMesh* UnFbx::FFbxImporter::ImportStaticMeshAsSingle(UObject* InParent, TA
 
 	if (ImportOptions->bImportEditableMesh)
 	{
-		MeshDescription = StaticMesh->MeshDescription;
+		MeshDescription = StaticMesh->GetMeshDescription(LODIndex);
 		if (MeshDescription == nullptr)
 		{
 			//Create private asset in the same package as the StaticMesh, and make sure reference are set to avoid GC
-			FString MeshDescriptionName = MeshName + TEXT("_MeshData");
-			StaticMesh->MeshDescription = NewObject<UMeshDescription>(StaticMesh, FName(*MeshDescriptionName), RF_NoFlags);
-			MeshDescription = StaticMesh->MeshDescription;
+			MeshDescription = NewObject<UMeshDescription>(StaticMesh, NAME_None, RF_NoFlags);
+			check(MeshDescription != nullptr);
+			StaticMesh->MeshDescriptions.Add(MeshDescription);
 		}
 	}
 
@@ -1895,14 +1904,48 @@ UStaticMesh* UnFbx::FFbxImporter::ImportStaticMeshAsSingle(UObject* InParent, TA
 					MeshPolygon.PolygonGroupID = FPolygonGroupID(MaterialMap[MeshPolygon.PolygonGroupID.GetValue()]);
 				}
 			}
-
+			TArray<FStaticMaterial> MaterialToAdd;
 			for (const FPolygonGroupID& PolygonGroupID : MeshDescription->PolygonGroups().GetElementIDs())
 			{
 				int32 MaterialIndex = PolygonGroupID.GetValue();
 				const FMeshPolygonGroup& MeshPolygonGroup = MeshDescription->GetPolygonGroup(PolygonGroupID);
-				int32 AddedIndex = StaticMesh->StaticMaterials.Add(FStaticMaterial(Cast<UMaterialInterface>(MeshPolygonGroup.MaterialAsset.TryLoad()), MeshPolygonGroup.MaterialSlotName, MeshPolygonGroup.ImportedMaterialSlotName));
+				FStaticMaterial StaticMaterial(Cast<UMaterialInterface>(MeshPolygonGroup.MaterialAsset.TryLoad()), MeshPolygonGroup.MaterialSlotName, MeshPolygonGroup.ImportedMaterialSlotName);
+				int32 AddedIndex = (LODIndex > 0) ? MaterialToAdd.Add(StaticMaterial) : StaticMesh->StaticMaterials.Add(StaticMaterial);
 				check(MaterialIndex == AddedIndex);
 			}
+			if (LODIndex > 0)
+			{
+				//Insert the new materials in the static mesh
+				//The build function will search for imported slot name to find the appropriate slot
+				int32 StaticMeshMaterialCount = StaticMesh->StaticMaterials.Num();
+				if (StaticMeshMaterialCount > 0)
+				{
+					for (int32 MaterialToAddIndex = 0; MaterialToAddIndex < MaterialToAdd.Num(); ++MaterialToAddIndex)
+					{
+						const FStaticMaterial& CandidateMaterial = MaterialToAdd[MaterialToAddIndex];
+						bool FoundExistingMaterial = false;
+						//Found matching existing material
+						for (int32 StaticMeshMaterialIndex = 0; StaticMeshMaterialIndex < StaticMeshMaterialCount; ++StaticMeshMaterialIndex)
+						{
+							const FStaticMaterial& StaticMeshMaterial = StaticMesh->StaticMaterials[StaticMeshMaterialIndex];
+							if (StaticMeshMaterial.MaterialInterface == CandidateMaterial.MaterialInterface)
+							{
+								FPolygonGroupID PolygonGroupID(MaterialToAddIndex);
+								FMeshPolygonGroup& MeshPolygonGroup = MeshDescription->GetPolygonGroup(PolygonGroupID);
+								//The LOD 0 existing material is driving the polygon groups slot detail, so all LOD mesh description use the same data.
+								MeshPolygonGroup.ImportedMaterialSlotName = StaticMeshMaterial.ImportedMaterialSlotName;
+								MeshPolygonGroup.MaterialSlotName = StaticMeshMaterial.MaterialSlotName;
+								FoundExistingMaterial = true;
+								break;
+							}
+						}
+						if (!FoundExistingMaterial)
+						{
+							StaticMesh->StaticMaterials.Add(CandidateMaterial);
+						}
+					}
+				}
+			}
 		}
 
 		// Setup default LOD settings based on the selected LOD group.
diff --git a/Engine/Source/Runtime/Engine/Classes/Engine/StaticMesh.h b/Engine/Source/Runtime/Engine/Classes/Engine/StaticMesh.h
--- a/Engine/Source/Runtime/Engine/Classes/Engine/StaticMesh.h
+++ b/Engine/Source/Runtime/Engine/Classes/Engine/StaticMesh.h
@@ -624,7 +624,10 @@ public:
 
 	/** The MeshDescription use to build the render data. */
 	UPROPERTY(VisibleAnywhere, Category = MeshDescription)
-	class UMeshDescription* MeshDescription;
+	TArray<class UMeshDescription*> MeshDescriptions;
+
+	ENGINE_API class UMeshDescription* GetMeshDescription(int32 LodIndex=0) const;
+	ENGINE_API int32 GetMeshDescriptionCount() const;
 
 	/** Pre-build navigation collision */
 	UPROPERTY(VisibleAnywhere, transient, duplicatetransient, Instanced, Category = Navigation)
@@ -747,6 +750,8 @@ public:
 	UFUNCTION(BlueprintCallable, Category = "StaticMesh")
 	ENGINE_API int32 GetMaterialIndex(FName MaterialSlotName) const;
 
+	ENGINE_API int32 GetMaterialIndexFromImportedMaterialSlotName(FName ImportedMaterialSlotName) const;
+
 	/**
 	 * Returns the render data to use for exporting the specified LOD. This method should always
 	 * be called when exporting a static mesh.
diff --git a/Engine/Source/Runtime/Engine/Private/StaticMesh.cpp b/Engine/Source/Runtime/Engine/Private/StaticMesh.cpp
--- a/Engine/Source/Runtime/Engine/Private/StaticMesh.cpp
+++ b/Engine/Source/Runtime/Engine/Private/StaticMesh.cpp
@@ -38,6 +38,7 @@
 
 #if WITH_EDITOR
 #include "RawMesh.h"
+#include "MeshDescription.h"
 #include "MeshUtilities.h"
 #include "DerivedDataCacheInterface.h"
 #endif // #if WITH_EDITOR
@@ -2168,6 +2169,21 @@ static FStaticMeshRenderData& GetPlatformStaticMeshRenderData(UStaticMesh* Mesh,
 	return *PlatformRenderData;
 }
 
+
+UMeshDescription* UStaticMesh::GetMeshDescription(int32 LodIndex) const
+{
+	if (!MeshDescriptions.IsValidIndex(LodIndex))
+	{
+		return nullptr;
+	}
+	return MeshDescriptions[LodIndex];
+}
+
+ENGINE_API int32 UStaticMesh::GetMeshDescriptionCount() const
+{
+	return MeshDescriptions.Num();
+}
+
 void UStaticMesh::CacheDerivedData()
 {
 	// Cache derived data for the running platform.
@@ -3382,6 +3398,19 @@ int32 UStaticMesh::GetMaterialIndex(FName MaterialSlotName) const
 	return -1;
 }
 
+int32 UStaticMesh::GetMaterialIndexFromImportedMaterialSlotName(FName ImportedMaterialSlotName) const
+{
+	for (int32 MaterialIndex = 0; MaterialIndex < StaticMaterials.Num(); ++MaterialIndex)
+	{
+		const FStaticMaterial &StaticMaterial = StaticMaterials[MaterialIndex];
+		if (StaticMaterial.ImportedMaterialSlotName == ImportedMaterialSlotName)
+		{
+			return MaterialIndex;
+		}
+	}
+	return -1;
+}
+
 /**
  * Returns the render data to use for exporting the specified LOD. This method should always
  * be called when exporting a static mesh.
diff --git a/Engine/Source/Runtime/Engine/Private/StaticMeshBuild.cpp b/Engine/Source/Runtime/Engine/Private/StaticMeshBuild.cpp
--- a/Engine/Source/Runtime/Engine/Private/StaticMeshBuild.cpp
+++ b/Engine/Source/Runtime/Engine/Private/StaticMeshBuild.cpp
@@ -60,7 +60,7 @@ void UStaticMesh::Build(bool bSilent, TArray<FText>* OutErrors)
 		return;
 
 	// If we're controlled by an editable mesh do not build. The editable mesh will build us
-	if (EditableMesh || MeshDescription)
+	if (EditableMesh || MeshDescriptions.Num() > 0)
 	{
 		return;
 	}
diff --git a/Engine/Source/Runtime/MeshBuilder/Private/StaticMeshBuilder.cpp b/Engine/Source/Runtime/MeshBuilder/Private/StaticMeshBuilder.cpp
--- a/Engine/Source/Runtime/MeshBuilder/Private/StaticMeshBuilder.cpp
+++ b/Engine/Source/Runtime/MeshBuilder/Private/StaticMeshBuilder.cpp
@@ -11,12 +11,12 @@ FStaticMeshBuilder::FStaticMeshBuilder()
 
 }
 
-int32 FStaticMeshBuilder::GetPolygonGroupTriangles(UStaticMesh* StaticMesh, TArray<FMeshTriangle>& OutTriangles, const FPolygonGroupID& PolygonGroupID)
+int32 FStaticMeshBuilder::GetPolygonGroupTriangles(const UMeshDescription *MeshDescription, TArray<FMeshTriangle>& OutTriangles, const FPolygonGroupID& PolygonGroupID)
 {
 	int32 TriangleCount = 0;
-	for (const FPolygonID& PolygonID : StaticMesh->MeshDescription->Polygons().GetElementIDs())
+	for (const FPolygonID& PolygonID : MeshDescription->Polygons().GetElementIDs())
 	{
-		FMeshPolygon MeshPolygon = StaticMesh->MeshDescription->GetPolygon(PolygonID);
+		FMeshPolygon MeshPolygon = MeshDescription->GetPolygon(PolygonID);
 		if (MeshPolygon.PolygonGroupID == PolygonGroupID)
 		{
 			for (const FMeshTriangle& MeshTriangle : MeshPolygon.Triangles)
@@ -31,139 +31,152 @@ int32 FStaticMeshBuilder::GetPolygonGroupTriangles(UStaticMesh* StaticMesh, TArr
 
 bool FStaticMeshBuilder::Build(UStaticMesh* StaticMesh)
 {
+	int32 MeshDescriptionCount = StaticMesh->GetMeshDescriptionCount();
+	if (MeshDescriptionCount <= 0)
+	{
+		//TODO: Warn the user that there is no mesh description data
+		return false;
+	}
+
 	OnBuildRenderMeshStart(StaticMesh, false);
 
 	FStaticMeshRenderData& StaticMeshRenderData = *StaticMesh->RenderData;
-	//TODO: support LOD in the mesh description
-	FStaticMeshLODResources& StaticMeshLOD = StaticMeshRenderData.LODResources[0];
 
-	// Build new vertex buffers
-	static TArray< FStaticMeshBuildVertex > StaticMeshBuildVertices;
-	StaticMeshBuildVertices.Reset();
+	for (int32 LodIndex = 0; LodIndex < MeshDescriptionCount; ++LodIndex)
+	{
+		const UMeshDescription* MeshDescription = StaticMesh->GetMeshDescription(LodIndex);
+		check(MeshDescription != nullptr);
 
-	static TArray< uint32 > IndexBuffer;
-	IndexBuffer.Reset();
+		FStaticMeshLODResources& StaticMeshLOD = StaticMeshRenderData.LODResources[LodIndex];
 
-	const FPolygonGroupArray& PolygonGroups = StaticMesh->MeshDescription->PolygonGroups();
+		// Build new vertex buffers
+		static TArray< FStaticMeshBuildVertex > StaticMeshBuildVertices;
+		StaticMeshBuildVertices.Reset();
 
-	StaticMeshLOD.Sections.Empty(PolygonGroups.Num());
+		static TArray< uint32 > IndexBuffer;
+		IndexBuffer.Reset();
 
-	bool bHasColor = false;
+		const FPolygonGroupArray& PolygonGroups = MeshDescription->PolygonGroups();
 
-	const FVertexArray& Vertices = StaticMesh->MeshDescription->Vertices();
-	const FVertexInstanceArray& VertexInstances = StaticMesh->MeshDescription->VertexInstances();
+		StaticMeshLOD.Sections.Empty(PolygonGroups.Num());
 
-	// set up vertex buffer elements
-	StaticMeshBuildVertices.SetNum(VertexInstances.GetArraySize());
+		bool bHasColor = false;
 
-	for (const FVertexInstanceID VertexInstanceID : VertexInstances.GetElementIDs())
-	{
-		const FMeshVertexInstance& VertexInstance = VertexInstances[VertexInstanceID];
+		const FVertexArray& Vertices = MeshDescription->Vertices();
+		const FVertexInstanceArray& VertexInstances = MeshDescription->VertexInstances();
 
-		if (VertexInstance.Color != FColor::White)
+		// set up vertex buffer elements
+		StaticMeshBuildVertices.SetNum(VertexInstances.GetArraySize());
+
+		for (const FVertexInstanceID VertexInstanceID : VertexInstances.GetElementIDs())
 		{
-			bHasColor = true;
-		}
+			const FMeshVertexInstance& VertexInstance = VertexInstances[VertexInstanceID];
 
-		FStaticMeshBuildVertex& StaticMeshVertex = StaticMeshBuildVertices[VertexInstanceID.GetValue()];
+			if (VertexInstance.Color != FColor::White)
+			{
+				bHasColor = true;
+			}
 
-		StaticMeshVertex.Position = Vertices[VertexInstance.VertexID].VertexPosition;
-		StaticMeshVertex.TangentX = VertexInstance.Tangent;
-		StaticMeshVertex.TangentY = FVector::CrossProduct(VertexInstance.Normal, VertexInstance.Tangent).GetSafeNormal() * VertexInstance.BinormalSign;
-		StaticMeshVertex.TangentZ = VertexInstance.Normal;
-		StaticMeshVertex.Color = VertexInstance.Color.ToFColor(true);
-		for (int32 UVIndex = 0; UVIndex < VertexInstance.VertexUVs.Num(); ++UVIndex)
-		{
-			StaticMeshVertex.UVs[UVIndex] = VertexInstance.VertexUVs[UVIndex];
+			FStaticMeshBuildVertex& StaticMeshVertex = StaticMeshBuildVertices[VertexInstanceID.GetValue()];
+
+			StaticMeshVertex.Position = Vertices[VertexInstance.VertexID].VertexPosition;
+			StaticMeshVertex.TangentX = VertexInstance.Tangent;
+			StaticMeshVertex.TangentY = FVector::CrossProduct(VertexInstance.Normal, VertexInstance.Tangent).GetSafeNormal() * VertexInstance.BinormalSign;
+			StaticMeshVertex.TangentZ = VertexInstance.Normal;
+			StaticMeshVertex.Color = VertexInstance.Color.ToFColor(true);
+			for (int32 UVIndex = 0; UVIndex < VertexInstance.VertexUVs.Num(); ++UVIndex)
+			{
+				StaticMeshVertex.UVs[UVIndex] = VertexInstance.VertexUVs[UVIndex];
+			}
 		}
-	}
 
-	// Set up index buffer
-	for (const FPolygonGroupID& PolygonGroupID : PolygonGroups.GetElementIDs())
-	{
-		TArray<FMeshTriangle> Triangles;
-		int32 SectionTriangleCount = GetPolygonGroupTriangles(StaticMesh, Triangles, PolygonGroupID);
+		// Set up index buffer
+		for (const FPolygonGroupID& PolygonGroupID : PolygonGroups.GetElementIDs())
+		{
+			TArray<FMeshTriangle> Triangles;
+			int32 SectionTriangleCount = GetPolygonGroupTriangles(MeshDescription, Triangles, PolygonGroupID);
 
-		const FMeshPolygonGroup& PolygonGroup = PolygonGroups[PolygonGroupID];
+			const FMeshPolygonGroup& PolygonGroup = PolygonGroups[PolygonGroupID];
 
-		// Create new rendering section
-		StaticMeshLOD.Sections.Add(FStaticMeshSection());
-		FStaticMeshSection& StaticMeshSection = StaticMeshLOD.Sections.Last();
+			// Create new rendering section
+			StaticMeshLOD.Sections.Add(FStaticMeshSection());
+			FStaticMeshSection& StaticMeshSection = StaticMeshLOD.Sections.Last();
 
-		StaticMeshSection.FirstIndex = IndexBuffer.Num();
-		StaticMeshSection.NumTriangles = SectionTriangleCount;
+			StaticMeshSection.FirstIndex = IndexBuffer.Num();
+			StaticMeshSection.NumTriangles = SectionTriangleCount;
 
-		const int32 MaterialIndex = StaticMesh->GetMaterialIndex(PolygonGroup.MaterialSlotName);
-		check(MaterialIndex != INDEX_NONE);
-		check(StaticMesh->StaticMaterials[MaterialIndex].MaterialInterface->GetPathName() == PolygonGroup.MaterialAsset.ToString());
-		StaticMeshSection.MaterialIndex = MaterialIndex;
-		StaticMeshSection.bEnableCollision = PolygonGroup.bEnableCollision;
-		StaticMeshSection.bCastShadow = PolygonGroup.bCastShadow;
+			const int32 MaterialIndex = StaticMesh->GetMaterialIndex(PolygonGroup.ImportedMaterialSlotName);
+			check(MaterialIndex != INDEX_NONE);
+			check(StaticMesh->StaticMaterials[MaterialIndex].MaterialInterface->GetPathName() == PolygonGroup.MaterialAsset.ToString());
+			StaticMeshSection.MaterialIndex = MaterialIndex;
+			StaticMeshSection.bEnableCollision = PolygonGroup.bEnableCollision;
+			StaticMeshSection.bCastShadow = PolygonGroup.bCastShadow;
 
-		if (Triangles.Num() > 0)
-		{
-			IndexBuffer.Reserve(IndexBuffer.Num() + Triangles.Num() * 3);
-			uint32 MinIndex = TNumericLimits< uint32 >::Max();
-			uint32 MaxIndex = TNumericLimits< uint32 >::Min();
-			for (int32 TriangleIndex = 0; TriangleIndex < Triangles.Num(); ++TriangleIndex)
+			if (Triangles.Num() > 0)
 			{
-				const FMeshTriangle& Triangle = Triangles[TriangleIndex];
-				for (int32 TriVert = 0; TriVert < 3; ++TriVert)
+				IndexBuffer.Reserve(IndexBuffer.Num() + Triangles.Num() * 3);
+				uint32 MinIndex = TNumericLimits< uint32 >::Max();
+				uint32 MaxIndex = TNumericLimits< uint32 >::Min();
+				for (int32 TriangleIndex = 0; TriangleIndex < Triangles.Num(); ++TriangleIndex)
 				{
-					const uint32 RenderingVertexIndex = Triangle.GetVertexInstanceID(TriVert).GetValue();
-					IndexBuffer.Add(RenderingVertexIndex);
-					MinIndex = FMath::Min(MinIndex, RenderingVertexIndex);
-					MaxIndex = FMath::Max(MaxIndex, RenderingVertexIndex);
+					const FMeshTriangle& Triangle = Triangles[TriangleIndex];
+					for (int32 TriVert = 0; TriVert < 3; ++TriVert)
+					{
+						const uint32 RenderingVertexIndex = Triangle.GetVertexInstanceID(TriVert).GetValue();
+						IndexBuffer.Add(RenderingVertexIndex);
+						MinIndex = FMath::Min(MinIndex, RenderingVertexIndex);
+						MaxIndex = FMath::Max(MaxIndex, RenderingVertexIndex);
+					}
 				}
+				StaticMeshSection.MinVertexIndex = MinIndex;
+				StaticMeshSection.MaxVertexIndex = MaxIndex;
+			}
+			else
+			{
+				// No triangles in this section
+				StaticMeshSection.MinVertexIndex = 0;
+				StaticMeshSection.MaxVertexIndex = 0;
 			}
-			StaticMeshSection.MinVertexIndex = MinIndex;
-			StaticMeshSection.MaxVertexIndex = MaxIndex;
-		}
-		else
-		{
-			// No triangles in this section
-			StaticMeshSection.MinVertexIndex = 0;
-			StaticMeshSection.MaxVertexIndex = 0;
 		}
-	}
 
-	// Figure out which index buffer stride we need
-	bool bNeeds32BitIndices = false;
-	for (const FStaticMeshSection& StaticMeshSection : StaticMeshLOD.Sections)
-	{
-		if (StaticMeshSection.MaxVertexIndex > TNumericLimits<uint16>::Max())
+		// Figure out which index buffer stride we need
+		bool bNeeds32BitIndices = false;
+		for (const FStaticMeshSection& StaticMeshSection : StaticMeshLOD.Sections)
 		{
-			bNeeds32BitIndices = true;
+			if (StaticMeshSection.MaxVertexIndex > TNumericLimits<uint16>::Max())
+			{
+				bNeeds32BitIndices = true;
+			}
 		}
-	}
-	const EIndexBufferStride::Type IndexBufferStride = bNeeds32BitIndices ? EIndexBufferStride::Force32Bit : EIndexBufferStride::Force16Bit;
+		const EIndexBufferStride::Type IndexBufferStride = bNeeds32BitIndices ? EIndexBufferStride::Force32Bit : EIndexBufferStride::Force16Bit;
 
-	StaticMeshLOD.PositionVertexBuffer.Init(StaticMeshBuildVertices);
-	StaticMeshLOD.VertexBuffer.Init(StaticMeshBuildVertices, 1); //TODO: use the real texture coordinnate count
+		StaticMeshLOD.PositionVertexBuffer.Init(StaticMeshBuildVertices);
+		StaticMeshLOD.VertexBuffer.Init(StaticMeshBuildVertices, 1); //TODO: use the real texture coordinnate count
 
-	if (bHasColor)
-	{
-		StaticMeshLOD.ColorVertexBuffer.Init(StaticMeshBuildVertices);
-	}
-	else
-	{
-		StaticMeshLOD.ColorVertexBuffer.InitFromSingleColor(FColor::White, StaticMeshBuildVertices.Num());
-	}
+		if (bHasColor)
+		{
+			StaticMeshLOD.ColorVertexBuffer.Init(StaticMeshBuildVertices);
+		}
+		else
+		{
+			StaticMeshLOD.ColorVertexBuffer.InitFromSingleColor(FColor::White, StaticMeshBuildVertices.Num());
+		}
 
-	StaticMeshLOD.IndexBuffer.SetIndices(IndexBuffer, IndexBufferStride);
+		StaticMeshLOD.IndexBuffer.SetIndices(IndexBuffer, IndexBufferStride);
 
-	// @todo mesheditor: support the other index buffer types
-	StaticMeshLOD.ReversedIndexBuffer.SetIndices(TArray< uint32 >(), IndexBufferStride);
-	StaticMeshLOD.DepthOnlyIndexBuffer.SetIndices(TArray< uint32 >(), IndexBufferStride);
-	StaticMeshLOD.ReversedDepthOnlyIndexBuffer.SetIndices(TArray< uint32 >(), IndexBufferStride);
-	StaticMeshLOD.WireframeIndexBuffer.SetIndices(TArray< uint32 >(), IndexBufferStride);
-	StaticMeshLOD.AdjacencyIndexBuffer.SetIndices(TArray< uint32 >(), IndexBufferStride);
+		// @todo mesheditor: support the other index buffer types
+		StaticMeshLOD.ReversedIndexBuffer.SetIndices(TArray< uint32 >(), IndexBufferStride);
+		StaticMeshLOD.DepthOnlyIndexBuffer.SetIndices(TArray< uint32 >(), IndexBufferStride);
+		StaticMeshLOD.ReversedDepthOnlyIndexBuffer.SetIndices(TArray< uint32 >(), IndexBufferStride);
+		StaticMeshLOD.WireframeIndexBuffer.SetIndices(TArray< uint32 >(), IndexBufferStride);
+		StaticMeshLOD.AdjacencyIndexBuffer.SetIndices(TArray< uint32 >(), IndexBufferStride);
 
-	StaticMeshLOD.bHasAdjacencyInfo = false;
-	StaticMeshLOD.bHasDepthOnlyIndices = false;
-	StaticMeshLOD.bHasReversedIndices = false;
-	StaticMeshLOD.bHasReversedDepthOnlyIndices = false;
-	StaticMeshLOD.DepthOnlyNumTriangles = 0;
+		StaticMeshLOD.bHasAdjacencyInfo = false;
+		StaticMeshLOD.bHasDepthOnlyIndices = false;
+		StaticMeshLOD.bHasReversedIndices = false;
+		StaticMeshLOD.bHasReversedDepthOnlyIndices = false;
+		StaticMeshLOD.DepthOnlyNumTriangles = 0;
+	} //End of LOD for loop
 
 	OnBuildRenderMeshFinish(StaticMesh, true);
 	return true;
@@ -209,10 +222,10 @@ void FStaticMeshBuilder::OnBuildRenderMeshFinish(UStaticMesh* StaticMesh, const
 	RecreateRenderStateContext.Reset();
 }
 
-bool FStaticMeshBuilder::IsOrphanedVertex(UStaticMesh* StaticMesh, const FVertexID VertexID) const
+bool FStaticMeshBuilder::IsOrphanedVertex(const UMeshDescription *MeshDescription, const FVertexID VertexID) const
 {
-	const FVertexInstanceArray& VertexInstances = StaticMesh->MeshDescription->VertexInstances();
-	for (const FVertexInstanceID VertexInstanceID : StaticMesh->MeshDescription->GetVertex(VertexID).VertexInstanceIDs)
+	const FVertexInstanceArray& VertexInstances = MeshDescription->VertexInstances();
+	for (const FVertexInstanceID VertexInstanceID : MeshDescription->GetVertex(VertexID).VertexInstanceIDs)
 	{
 		if (VertexInstances[VertexInstanceID].ConnectedPolygons.Num() > 0)
 		{
@@ -225,6 +238,9 @@ bool FStaticMeshBuilder::IsOrphanedVertex(UStaticMesh* StaticMesh, const FVertex
 
 void FStaticMeshBuilder::UpdateBounds(UStaticMesh* StaticMesh)
 {
+	const UMeshDescription* MeshDescription = StaticMesh->GetMeshDescription();
+	check(MeshDescription != nullptr);
+
 	// Compute a new bounding box
 	// @todo mesheditor perf: During the final modification, only do this if the bounds may have changed (need hinting)
 	FBoxSphereBounds BoundingBoxAndSphere;
@@ -237,12 +253,12 @@ void FStaticMeshBuilder::UpdateBounds(UStaticMesh* StaticMesh)
 	// 2) cache bounding box per polygon
 	// There are other cases where having polygon adjacency information (1) might be useful, so it's maybe worth considering.
 
-	const FVertexArray& Vertices = StaticMesh->MeshDescription->Vertices();
-	const FVertexInstanceArray& VertexInstances = StaticMesh->MeshDescription->VertexInstances();
+	const FVertexArray& Vertices = MeshDescription->Vertices();
+	const FVertexInstanceArray& VertexInstances = MeshDescription->VertexInstances();
 
 	for (const FVertexID VertexID : Vertices.GetElementIDs())
 	{
-		if (!IsOrphanedVertex(StaticMesh, VertexID))
+		if (!IsOrphanedVertex(MeshDescription, VertexID))
 		{
 			BoundingBox += Vertices[VertexID].VertexPosition;
 		}
@@ -255,7 +271,7 @@ void FStaticMeshBuilder::UpdateBounds(UStaticMesh* StaticMesh)
 
 	for (const FVertexID VertexID : Vertices.GetElementIDs())
 	{
-		if (!IsOrphanedVertex(StaticMesh, VertexID))
+		if (!IsOrphanedVertex(MeshDescription, VertexID))
 		{
 			const FVector VertexPosition = Vertices[VertexID].VertexPosition;
 			BoundingBoxAndSphere.SphereRadius = FMath::Max((VertexPosition - BoundingBoxAndSphere.Origin).Size(), BoundingBoxAndSphere.SphereRadius);
diff --git a/Engine/Source/Runtime/MeshBuilder/Public/StaticMeshBuilder.h b/Engine/Source/Runtime/MeshBuilder/Public/StaticMeshBuilder.h
--- a/Engine/Source/Runtime/MeshBuilder/Public/StaticMeshBuilder.h
+++ b/Engine/Source/Runtime/MeshBuilder/Public/StaticMeshBuilder.h
@@ -25,9 +25,9 @@ private:
 	void OnBuildRenderMeshStart(class UStaticMesh* StaticMesh, const bool bInvalidateLighting);
 	void OnBuildRenderMeshFinish(class UStaticMesh* StaticMesh, const bool bRebuildBoundsAndCollision);
 	
-	int32 GetPolygonGroupTriangles(class UStaticMesh* StaticMesh, TArray<FMeshTriangle>& OutTriangles, const FPolygonGroupID& PolygonGroupID);
+	int32 GetPolygonGroupTriangles(const class UMeshDescription *MeshDescription, TArray<FMeshTriangle>& OutTriangles, const FPolygonGroupID& PolygonGroupID);
 
-	bool IsOrphanedVertex(class UStaticMesh* StaticMesh, const FVertexID VertexID) const;
+	bool IsOrphanedVertex(const class UMeshDescription *MeshDescription, const FVertexID VertexID) const;
 	void UpdateBounds(class UStaticMesh* StaticMesh);
 	void UpdateCollision(UStaticMesh *StaticMesh);
 
