Implement build options basic feature
-We now build all the index buffer optimization array.

diff --git a/Engine/Source/Runtime/MeshBuilder/MeshBuilder.Build.cs b/Engine/Source/Runtime/MeshBuilder/MeshBuilder.Build.cs
--- a/Engine/Source/Runtime/MeshBuilder/MeshBuilder.Build.cs
+++ b/Engine/Source/Runtime/MeshBuilder/MeshBuilder.Build.cs
@@ -19,6 +19,9 @@ namespace UnrealBuildTool.Rules
                     "MeshDescription",
                 }
 			);
-		}
+
+            AddEngineThirdPartyPrivateStaticDependencies(Target, "nvTriStrip");
+            AddEngineThirdPartyPrivateStaticDependencies(Target, "nvTessLib");
+        }
 	}
 }
diff --git a/Engine/Source/Runtime/MeshBuilder/Private/BuildOptimizationHelper.h b/Engine/Source/Runtime/MeshBuilder/Private/BuildOptimizationHelper.h
new file mode 100644
--- /dev/null
+++ b/Engine/Source/Runtime/MeshBuilder/Private/BuildOptimizationHelper.h
@@ -0,0 +1,388 @@
+// Copyright 1998-2017 Epic Games, Inc. All Rights Reserved.
+#pragma once
+
+#include "CoreMinimal.h"
+#include "CoreTypes.h"
+#include "NvTriStrip.h"
+#include "nvtess.h"
+#include "Components.h"
+#include "StaticMeshResources.h"
+
+class BuildOptimizationHelper
+{
+public:
+	/** Helper struct for building acceleration structures. */
+	struct FIndexAndZ
+	{
+		float Z;
+		int32 Index;
+
+		/** Default constructor. */
+		FIndexAndZ() {}
+
+		/** Initialization constructor. */
+		FIndexAndZ(int32 InIndex, FVector V)
+		{
+			Z = 0.30f * V.X + 0.33f * V.Y + 0.37f * V.Z;
+			Index = InIndex;
+		}
+	};
+
+	/** Sorting function for vertex Z/index pairs. */
+	struct FCompareIndexAndZ
+	{
+		FORCEINLINE bool operator()(FIndexAndZ const& A, FIndexAndZ const& B) const { return A.Z < B.Z; }
+	};
+
+	/**
+	* Returns true if the specified points are about equal
+	*/
+	static inline bool PointsEqual(const FVector& V1, const FVector& V2, float ComparisonThreshold)
+	{
+		if (FMath::Abs(V1.X - V2.X) > ComparisonThreshold
+			|| FMath::Abs(V1.Y - V2.Y) > ComparisonThreshold
+			|| FMath::Abs(V1.Z - V2.Z) > ComparisonThreshold)
+		{
+			return false;
+		}
+		return true;
+	}
+
+	struct FMeshEdge
+	{
+		int32	Vertices[2];
+		int32	Faces[2];
+	};
+
+	/**
+	* This helper class builds the edge list for a mesh. It uses a hash of vertex
+	* positions to edges sharing that vertex to remove the n^2 searching of all
+	* previously added edges. This class is templatized so it can be used with
+	* either static mesh or skeletal mesh vertices
+	*/
+	template <class VertexClass> class TEdgeBuilder
+	{
+	protected:
+		/**
+		* The list of indices to build the edge data from
+		*/
+		const TArray<uint32>& Indices;
+		/**
+		* The array of verts for vertex position comparison
+		*/
+		const TArray<VertexClass>& Vertices;
+		/**
+		* The array of edges to create
+		*/
+		TArray<FMeshEdge>& Edges;
+		/**
+		* List of edges that start with a given vertex
+		*/
+		TMultiMap<FVector, FMeshEdge*> VertexToEdgeList;
+
+		/**
+		* This function determines whether a given edge matches or not. It must be
+		* provided by derived classes since they have the specific information that
+		* this class doesn't know about (vertex info, influences, etc)
+		*
+		* @param Index1 The first index of the edge being checked
+		* @param Index2 The second index of the edge
+		* @param OtherEdge The edge to compare. Was found via the map
+		*
+		* @return true if the edge is a match, false otherwise
+		*/
+		virtual bool DoesEdgeMatch(int32 Index1, int32 Index2, FMeshEdge* OtherEdge) = 0;
+
+		/**
+		* Searches the list of edges to see if this one matches an existing and
+		* returns a pointer to it if it does
+		*
+		* @param Index1 the first index to check for
+		* @param Index2 the second index to check for
+		*
+		* @return NULL if no edge was found, otherwise the edge that was found
+		*/
+		inline FMeshEdge* FindOppositeEdge(int32 Index1, int32 Index2)
+		{
+			FMeshEdge* Edge = NULL;
+			TArray<FMeshEdge*> EdgeList;
+			// Search the hash for a corresponding vertex
+			VertexToEdgeList.MultiFind(Vertices[Index2].Position, EdgeList);
+			// Now search through the array for a match or not
+			for (int32 EdgeIndex = 0; EdgeIndex < EdgeList.Num() && Edge == NULL;
+				EdgeIndex++)
+			{
+				FMeshEdge* OtherEdge = EdgeList[EdgeIndex];
+				// See if this edge matches the passed in edge
+				if (OtherEdge != NULL && DoesEdgeMatch(Index1, Index2, OtherEdge))
+				{
+					// We have a match
+					Edge = OtherEdge;
+				}
+			}
+			return Edge;
+		}
+
+		/**
+		* Updates an existing edge if found or adds the new edge to the list
+		*
+		* @param Index1 the first index in the edge
+		* @param Index2 the second index in the edge
+		* @param Triangle the triangle that this edge was found in
+		*/
+		inline void AddEdge(int32 Index1, int32 Index2, int32 Triangle)
+		{
+			// If this edge matches another then just fill the other triangle
+			// otherwise add it
+			FMeshEdge* OtherEdge = FindOppositeEdge(Index1, Index2);
+			if (OtherEdge == NULL)
+			{
+				// Add a new edge to the array
+				int32 EdgeIndex = Edges.AddZeroed();
+				Edges[EdgeIndex].Vertices[0] = Index1;
+				Edges[EdgeIndex].Vertices[1] = Index2;
+				Edges[EdgeIndex].Faces[0] = Triangle;
+				Edges[EdgeIndex].Faces[1] = -1;
+				// Also add this edge to the hash for faster searches
+				// NOTE: This relies on the array never being realloced!
+				VertexToEdgeList.Add(Vertices[Index1].Position, &Edges[EdgeIndex]);
+			}
+			else
+			{
+				OtherEdge->Faces[1] = Triangle;
+			}
+		}
+
+	public:
+		/**
+		* Initializes the values for the code that will build the mesh edge list
+		*/
+		TEdgeBuilder(const TArray<uint32>& InIndices,
+			const TArray<VertexClass>& InVertices,
+			TArray<FMeshEdge>& OutEdges) :
+			Indices(InIndices), Vertices(InVertices), Edges(OutEdges)
+		{
+			// Presize the array so that there are no extra copies being done
+			// when adding edges to it
+			Edges.Empty(Indices.Num());
+		}
+
+		/**
+		* Virtual dtor
+		*/
+		virtual ~TEdgeBuilder() {}
+
+
+		/**
+		* Uses a hash of indices to edge lists so that it can avoid the n^2 search
+		* through the full edge list
+		*/
+		void FindEdges(void)
+		{
+			// @todo Handle something other than trilists when building edges
+			int32 TriangleCount = Indices.Num() / 3;
+			int32 EdgeCount = 0;
+			// Work through all triangles building the edges
+			for (int32 Triangle = 0; Triangle < TriangleCount; Triangle++)
+			{
+				// Determine the starting index
+				int32 TriangleIndex = Triangle * 3;
+				// Get the indices for the triangle
+				int32 Index1 = Indices[TriangleIndex];
+				int32 Index2 = Indices[TriangleIndex + 1];
+				int32 Index3 = Indices[TriangleIndex + 2];
+				// Add the first to second edge
+				AddEdge(Index1, Index2, Triangle);
+				// Now add the second to third
+				AddEdge(Index2, Index3, Triangle);
+				// Add the third to first edge
+				AddEdge(Index3, Index1, Triangle);
+			}
+		}
+	};
+
+
+	/**
+	* This is the static mesh specific version for finding edges
+	*/
+	class FStaticMeshEdgeBuilder : public TEdgeBuilder<FStaticMeshBuildVertex>
+	{
+	public:
+		/**
+		* Constructor that passes all work to the parent class
+		*/
+		FStaticMeshEdgeBuilder(const TArray<uint32>& InIndices,
+			const TArray<FStaticMeshBuildVertex>& InVertices,
+			TArray<FMeshEdge>& OutEdges) :
+			TEdgeBuilder<FStaticMeshBuildVertex>(InIndices, InVertices, OutEdges)
+		{
+		}
+
+		/**
+		* This function determines whether a given edge matches or not for a static mesh
+		*
+		* @param Index1 The first index of the edge being checked
+		* @param Index2 The second index of the edge
+		* @param OtherEdge The edge to compare. Was found via the map
+		*
+		* @return true if the edge is a match, false otherwise
+		*/
+		inline bool DoesEdgeMatch(int32 Index1, int32 Index2, FMeshEdge* OtherEdge)
+		{
+			return Vertices[OtherEdge->Vertices[1]].Position == Vertices[Index1].Position &&
+				OtherEdge->Faces[1] == -1;
+		}
+	};
+
+	/**
+	* Provides static mesh render data to the NVIDIA tessellation library.
+	*/
+	class FStaticMeshNvRenderBuffer : public nv::RenderBuffer
+	{
+	public:
+
+		/** Construct from static mesh render buffers. */
+		FStaticMeshNvRenderBuffer(
+			const FPositionVertexBuffer& InPositionVertexBuffer,
+			const FStaticMeshVertexBuffer& InVertexBuffer,
+			const TArray<uint32>& Indices)
+			: PositionVertexBuffer(InPositionVertexBuffer)
+			, VertexBuffer(InVertexBuffer)
+		{
+			check(PositionVertexBuffer.GetNumVertices() == VertexBuffer.GetNumVertices());
+			mIb = new nv::IndexBuffer((void*)Indices.GetData(), nv::IBT_U32, Indices.Num(), false);
+		}
+
+		/** Retrieve the position and first texture coordinate of the specified index. */
+		virtual nv::Vertex getVertex(unsigned int Index) const
+		{
+			nv::Vertex Vertex;
+
+			check(Index < PositionVertexBuffer.GetNumVertices());
+
+			const FVector& Position = PositionVertexBuffer.VertexPosition(Index);
+			Vertex.pos.x = Position.X;
+			Vertex.pos.y = Position.Y;
+			Vertex.pos.z = Position.Z;
+
+			if (VertexBuffer.GetNumTexCoords())
+			{
+				const FVector2D UV = VertexBuffer.GetVertexUV(Index, 0);
+				Vertex.uv.x = UV.X;
+				Vertex.uv.y = UV.Y;
+			}
+			else
+			{
+				Vertex.uv.x = 0.0f;
+				Vertex.uv.y = 0.0f;
+			}
+
+			return Vertex;
+		}
+
+	private:
+
+		/** The position vertex buffer for the static mesh. */
+		const FPositionVertexBuffer& PositionVertexBuffer;
+
+		/** The vertex buffer for the static mesh. */
+		const FStaticMeshVertexBuffer& VertexBuffer;
+
+		/** Copying is forbidden. */
+		FStaticMeshNvRenderBuffer(const FStaticMeshNvRenderBuffer&);
+		FStaticMeshNvRenderBuffer& operator=(const FStaticMeshNvRenderBuffer&);
+	};
+
+	static void BuildDepthOnlyIndexBuffer(
+		TArray<uint32>& OutDepthIndices,
+		const TArray<FStaticMeshBuildVertex>& InVertices,
+		const TArray<uint32>& InIndices,
+		const TArray<FStaticMeshSection>& InSections
+	)
+	{
+		int32 NumVertices = InVertices.Num();
+		if (InIndices.Num() <= 0 || NumVertices <= 0)
+		{
+			OutDepthIndices.Empty();
+			return;
+		}
+
+		// Create a mapping of index -> first overlapping index to accelerate the construction of the shadow index buffer.
+		TArray<FIndexAndZ> VertIndexAndZ;
+		VertIndexAndZ.Empty(NumVertices);
+		for (int32 VertIndex = 0; VertIndex < NumVertices; VertIndex++)
+		{
+			new(VertIndexAndZ)FIndexAndZ(VertIndex, InVertices[VertIndex].Position);
+		}
+		VertIndexAndZ.Sort(FCompareIndexAndZ());
+
+		// Setup the index map. 0xFFFFFFFF == not set.
+		TArray<uint32> IndexMap;
+		IndexMap.AddUninitialized(NumVertices);
+		FMemory::Memset(IndexMap.GetData(), 0xFF, NumVertices * sizeof(uint32));
+
+		// Search for duplicates, quickly!
+		for (int32 i = 0; i < VertIndexAndZ.Num(); i++)
+		{
+			uint32 SrcIndex = VertIndexAndZ[i].Index;
+			float Z = VertIndexAndZ[i].Z;
+			IndexMap[SrcIndex] = FMath::Min(IndexMap[SrcIndex], SrcIndex);
+
+			// Search forward since we add pairs both ways.
+			for (int32 j = i + 1; j < VertIndexAndZ.Num(); j++)
+			{
+				if (FMath::Abs(VertIndexAndZ[j].Z - Z) > THRESH_POINTS_ARE_SAME * 4.01f)
+					break; // can't be any more dups
+
+				uint32 OtherIndex = VertIndexAndZ[j].Index;
+				if (PointsEqual(InVertices[SrcIndex].Position, InVertices[OtherIndex].Position,/*bUseEpsilonCompare=*/ false))
+				{
+					IndexMap[SrcIndex] = FMath::Min(IndexMap[SrcIndex], OtherIndex);
+					IndexMap[OtherIndex] = FMath::Min(IndexMap[OtherIndex], SrcIndex);
+				}
+			}
+		}
+
+		// Build the depth-only index buffer by remapping all indices to the first overlapping
+		// vertex in the vertex buffer.
+		OutDepthIndices.Empty();
+		for (int32 SectionIndex = 0; SectionIndex < InSections.Num(); ++SectionIndex)
+		{
+			const FStaticMeshSection& Section = InSections[SectionIndex];
+			int32 FirstIndex = Section.FirstIndex;
+			int32 LastIndex = FirstIndex + Section.NumTriangles * 3;
+			for (int32 SrcIndex = FirstIndex; SrcIndex < LastIndex; ++SrcIndex)
+			{
+				uint32 VertIndex = InIndices[SrcIndex];
+				OutDepthIndices.Add(IndexMap[VertIndex]);
+			}
+		}
+	}
+
+	static void BuildStaticAdjacencyIndexBuffer(
+		const FPositionVertexBuffer& PositionVertexBuffer,
+		const FStaticMeshVertexBuffer& VertexBuffer,
+		const TArray<uint32>& Indices,
+		TArray<uint32>& OutPnAenIndices
+	)
+	{
+		if (Indices.Num())
+		{
+			FStaticMeshNvRenderBuffer StaticMeshRenderBuffer(PositionVertexBuffer, VertexBuffer, Indices);
+			nv::IndexBuffer* PnAENIndexBuffer = nv::tess::buildTessellationBuffer(&StaticMeshRenderBuffer, nv::DBM_PnAenDominantCorner, true);
+			check(PnAENIndexBuffer);
+			const int32 IndexCount = (int32)PnAENIndexBuffer->getLength();
+			OutPnAenIndices.Empty(IndexCount);
+			OutPnAenIndices.AddUninitialized(IndexCount);
+			for (int32 Index = 0; Index < IndexCount; ++Index)
+			{
+				OutPnAenIndices[Index] = (*PnAENIndexBuffer)[Index];
+			}
+			delete PnAENIndexBuffer;
+		}
+		else
+		{
+			OutPnAenIndices.Empty();
+		}
+	}
+};
\ No newline at end of file
diff --git a/Engine/Source/Runtime/MeshBuilder/Private/MeshBuilderModule.cpp b/Engine/Source/Runtime/MeshBuilder/Private/MeshBuilderModule.cpp
--- a/Engine/Source/Runtime/MeshBuilder/Private/MeshBuilderModule.cpp
+++ b/Engine/Source/Runtime/MeshBuilder/Private/MeshBuilderModule.cpp
@@ -38,7 +38,7 @@ bool FMeshBuilderModule::BuildMesh(class UObject* Mesh)
 	if (StaticMesh != nullptr)
 	{
 		//Call the static mesh builder
-		return FStaticMeshBuilder::Get().Build(StaticMesh);
+		return FStaticMeshBuilder().Build(StaticMesh);
 	}
 	return false;
 }
diff --git a/Engine/Source/Runtime/MeshBuilder/Private/StaticMeshBuilder.cpp b/Engine/Source/Runtime/MeshBuilder/Private/StaticMeshBuilder.cpp
--- a/Engine/Source/Runtime/MeshBuilder/Private/StaticMeshBuilder.cpp
+++ b/Engine/Source/Runtime/MeshBuilder/Private/StaticMeshBuilder.cpp
@@ -5,30 +5,25 @@
 #include "StaticMeshResources.h"
 #include "PhysicsEngine/BodySetup.h"
 #include "MeshDescription.h"
+#include "BuildOptimizationHelper.h"
+#include "Components.h"
+
+//////////////////////////////////////////////////////////////////////////
+//Local functions definition
+int32 GetPolygonGroupTriangles(const class UMeshDescription *MeshDescription, TArray<FMeshTriangle>& OutTriangles, const FPolygonGroupID& PolygonGroupID);
+bool IsOrphanedVertex(const class UMeshDescription *MeshDescription, const FVertexID VertexID);
+void UpdateBounds(class UStaticMesh* StaticMesh);
+void UpdateCollision(UStaticMesh *StaticMesh);
+void BuildVertexBuffer(const class UMeshDescription* MeshDescription, struct FStaticMeshLODResources& StaticMeshLOD, const struct FMeshBuildSettings& LODBuildSettings, TArray< FStaticMeshBuildVertex >& StaticMeshBuildVertices);
+void BuildIndexBuffer(class UStaticMesh *StaticMesh, const class UMeshDescription* MeshDescription, TArray< uint32 >& IndexBuffer, struct FStaticMeshLODResources& StaticMeshLOD);
+void BuildAllBufferOptimizations(struct FStaticMeshLODResources& StaticMeshLOD, const struct FMeshBuildSettings& LODBuildSettings, TArray< uint32 >& IndexBuffer, bool bNeeds32BitIndices, TArray< FStaticMeshBuildVertex >& StaticMeshBuildVertices);
+//////////////////////////////////////////////////////////////////////////
 
 FStaticMeshBuilder::FStaticMeshBuilder()
 {
 
 }
 
-int32 FStaticMeshBuilder::GetPolygonGroupTriangles(const UMeshDescription *MeshDescription, TArray<FMeshTriangle>& OutTriangles, const FPolygonGroupID& PolygonGroupID)
-{
-	int32 TriangleCount = 0;
-	for (const FPolygonID& PolygonID : MeshDescription->Polygons().GetElementIDs())
-	{
-		FMeshPolygon MeshPolygon = MeshDescription->GetPolygon(PolygonID);
-		if (MeshPolygon.PolygonGroupID == PolygonGroupID)
-		{
-			for (const FMeshTriangle& MeshTriangle : MeshPolygon.Triangles)
-			{
-				OutTriangles.Add(MeshTriangle);
-				TriangleCount++;
-			}
-		}
-	}
-	return TriangleCount;
-}
-
 bool FStaticMeshBuilder::Build(UStaticMesh* StaticMesh)
 {
 	int32 MeshDescriptionCount = StaticMesh->GetMeshDescriptionCount();
@@ -41,104 +36,32 @@ bool FStaticMeshBuilder::Build(UStaticMesh* StaticMesh)
 	OnBuildRenderMeshStart(StaticMesh, false);
 
 	FStaticMeshRenderData& StaticMeshRenderData = *StaticMesh->RenderData;
-
 	for (int32 LodIndex = 0; LodIndex < MeshDescriptionCount; ++LodIndex)
 	{
 		const UMeshDescription* MeshDescription = StaticMesh->GetMeshDescription(LodIndex);
 		check(MeshDescription != nullptr);
+		
+		const FPolygonGroupArray& PolygonGroups = MeshDescription->PolygonGroups();
 
 		FStaticMeshLODResources& StaticMeshLOD = StaticMeshRenderData.LODResources[LodIndex];
+		const FMeshBuildSettings& LODBuildSettings = StaticMesh->SourceModels[LodIndex].BuildSettings;
+		
+		//TODO: discover degenerate triangle with this threshold
+		float VertexComparisonThreshold = LODBuildSettings.bRemoveDegenerates ? THRESH_POINTS_ARE_SAME : 0.0f;
 
-		// Build new vertex buffers
+		//Build new vertex buffers
 		static TArray< FStaticMeshBuildVertex > StaticMeshBuildVertices;
 		StaticMeshBuildVertices.Reset();
 
 		static TArray< uint32 > IndexBuffer;
 		IndexBuffer.Reset();
 
-		const FPolygonGroupArray& PolygonGroups = MeshDescription->PolygonGroups();
-
 		StaticMeshLOD.Sections.Empty(PolygonGroups.Num());
 
-		bool bHasColor = false;
-
-		const FVertexArray& Vertices = MeshDescription->Vertices();
-		const FVertexInstanceArray& VertexInstances = MeshDescription->VertexInstances();
-
-		// set up vertex buffer elements
-		StaticMeshBuildVertices.SetNum(VertexInstances.GetArraySize());
-
-		for (const FVertexInstanceID VertexInstanceID : VertexInstances.GetElementIDs())
-		{
-			const FMeshVertexInstance& VertexInstance = VertexInstances[VertexInstanceID];
-
-			if (VertexInstance.Color != FColor::White)
-			{
-				bHasColor = true;
-			}
-
-			FStaticMeshBuildVertex& StaticMeshVertex = StaticMeshBuildVertices[VertexInstanceID.GetValue()];
-
-			StaticMeshVertex.Position = Vertices[VertexInstance.VertexID].VertexPosition;
-			StaticMeshVertex.TangentX = VertexInstance.Tangent;
-			StaticMeshVertex.TangentY = FVector::CrossProduct(VertexInstance.Normal, VertexInstance.Tangent).GetSafeNormal() * VertexInstance.BinormalSign;
-			StaticMeshVertex.TangentZ = VertexInstance.Normal;
-			StaticMeshVertex.Color = VertexInstance.Color.ToFColor(true);
-			for (int32 UVIndex = 0; UVIndex < VertexInstance.VertexUVs.Num(); ++UVIndex)
-			{
-				StaticMeshVertex.UVs[UVIndex] = VertexInstance.VertexUVs[UVIndex];
-			}
-		}
-
-		// Set up index buffer
-		for (const FPolygonGroupID& PolygonGroupID : PolygonGroups.GetElementIDs())
-		{
-			TArray<FMeshTriangle> Triangles;
-			int32 SectionTriangleCount = GetPolygonGroupTriangles(MeshDescription, Triangles, PolygonGroupID);
-
-			const FMeshPolygonGroup& PolygonGroup = PolygonGroups[PolygonGroupID];
-
-			// Create new rendering section
-			StaticMeshLOD.Sections.Add(FStaticMeshSection());
-			FStaticMeshSection& StaticMeshSection = StaticMeshLOD.Sections.Last();
-
-			StaticMeshSection.FirstIndex = IndexBuffer.Num();
-			StaticMeshSection.NumTriangles = SectionTriangleCount;
-
-			const int32 MaterialIndex = StaticMesh->GetMaterialIndex(PolygonGroup.ImportedMaterialSlotName);
-			check(MaterialIndex != INDEX_NONE);
-			check(StaticMesh->StaticMaterials[MaterialIndex].MaterialInterface->GetPathName() == PolygonGroup.MaterialAsset.ToString());
-			StaticMeshSection.MaterialIndex = MaterialIndex;
-			StaticMeshSection.bEnableCollision = PolygonGroup.bEnableCollision;
-			StaticMeshSection.bCastShadow = PolygonGroup.bCastShadow;
-
-			if (Triangles.Num() > 0)
-			{
-				IndexBuffer.Reserve(IndexBuffer.Num() + Triangles.Num() * 3);
-				uint32 MinIndex = TNumericLimits< uint32 >::Max();
-				uint32 MaxIndex = TNumericLimits< uint32 >::Min();
-				for (int32 TriangleIndex = 0; TriangleIndex < Triangles.Num(); ++TriangleIndex)
-				{
-					const FMeshTriangle& Triangle = Triangles[TriangleIndex];
-					for (int32 TriVert = 0; TriVert < 3; ++TriVert)
-					{
-						const uint32 RenderingVertexIndex = Triangle.GetVertexInstanceID(TriVert).GetValue();
-						IndexBuffer.Add(RenderingVertexIndex);
-						MinIndex = FMath::Min(MinIndex, RenderingVertexIndex);
-						MaxIndex = FMath::Max(MaxIndex, RenderingVertexIndex);
-					}
-				}
-				StaticMeshSection.MinVertexIndex = MinIndex;
-				StaticMeshSection.MaxVertexIndex = MaxIndex;
-			}
-			else
-			{
-				// No triangles in this section
-				StaticMeshSection.MinVertexIndex = 0;
-				StaticMeshSection.MaxVertexIndex = 0;
-			}
-		}
+		BuildVertexBuffer(MeshDescription, StaticMeshLOD, LODBuildSettings, StaticMeshBuildVertices);
 
+		BuildIndexBuffer(StaticMesh, MeshDescription, IndexBuffer, StaticMeshLOD);
+		
 		// Figure out which index buffer stride we need
 		bool bNeeds32BitIndices = false;
 		for (const FStaticMeshSection& StaticMeshSection : StaticMeshLOD.Sections)
@@ -149,40 +72,15 @@ bool FStaticMeshBuilder::Build(UStaticMesh* StaticMesh)
 			}
 		}
 		const EIndexBufferStride::Type IndexBufferStride = bNeeds32BitIndices ? EIndexBufferStride::Force32Bit : EIndexBufferStride::Force16Bit;
-
-		StaticMeshLOD.PositionVertexBuffer.Init(StaticMeshBuildVertices);
-		StaticMeshLOD.VertexBuffer.Init(StaticMeshBuildVertices, 1); //TODO: use the real texture coordinnate count
-
-		if (bHasColor)
-		{
-			StaticMeshLOD.ColorVertexBuffer.Init(StaticMeshBuildVertices);
-		}
-		else
-		{
-			StaticMeshLOD.ColorVertexBuffer.InitFromSingleColor(FColor::White, StaticMeshBuildVertices.Num());
-		}
-
 		StaticMeshLOD.IndexBuffer.SetIndices(IndexBuffer, IndexBufferStride);
 
-		// @todo mesheditor: support the other index buffer types
-		StaticMeshLOD.ReversedIndexBuffer.SetIndices(TArray< uint32 >(), IndexBufferStride);
-		StaticMeshLOD.DepthOnlyIndexBuffer.SetIndices(TArray< uint32 >(), IndexBufferStride);
-		StaticMeshLOD.ReversedDepthOnlyIndexBuffer.SetIndices(TArray< uint32 >(), IndexBufferStride);
-		StaticMeshLOD.WireframeIndexBuffer.SetIndices(TArray< uint32 >(), IndexBufferStride);
-		StaticMeshLOD.AdjacencyIndexBuffer.SetIndices(TArray< uint32 >(), IndexBufferStride);
-
-		StaticMeshLOD.bHasAdjacencyInfo = false;
-		StaticMeshLOD.bHasDepthOnlyIndices = false;
-		StaticMeshLOD.bHasReversedIndices = false;
-		StaticMeshLOD.bHasReversedDepthOnlyIndices = false;
-		StaticMeshLOD.DepthOnlyNumTriangles = 0;
+		BuildAllBufferOptimizations(StaticMeshLOD, LODBuildSettings, IndexBuffer, bNeeds32BitIndices, StaticMeshBuildVertices);
 	} //End of LOD for loop
 
 	OnBuildRenderMeshFinish(StaticMesh, true);
 	return true;
 }
 
-
 void FStaticMeshBuilder::OnBuildRenderMeshStart(UStaticMesh* StaticMesh, const bool bInvalidateLighting)
 {
 	// We may already have a lock on the rendering resources, if it wasn't released the last time we called EndModification()
@@ -206,7 +104,7 @@ void FStaticMeshBuilder::OnBuildRenderMeshStart(UStaticMesh* StaticMesh, const b
 
 void FStaticMeshBuilder::OnBuildRenderMeshFinish(UStaticMesh* StaticMesh, const bool bRebuildBoundsAndCollision)
 {
-	
+
 
 	if (bRebuildBoundsAndCollision)
 	{
@@ -222,7 +120,25 @@ void FStaticMeshBuilder::OnBuildRenderMeshFinish(UStaticMesh* StaticMesh, const
 	RecreateRenderStateContext.Reset();
 }
 
-bool FStaticMeshBuilder::IsOrphanedVertex(const UMeshDescription *MeshDescription, const FVertexID VertexID) const
+int32 GetPolygonGroupTriangles(const UMeshDescription *MeshDescription, TArray<FMeshTriangle>& OutTriangles, const FPolygonGroupID& PolygonGroupID)
+{
+	int32 TriangleCount = 0;
+	for (const FPolygonID& PolygonID : MeshDescription->Polygons().GetElementIDs())
+	{
+		FMeshPolygon MeshPolygon = MeshDescription->GetPolygon(PolygonID);
+		if (MeshPolygon.PolygonGroupID == PolygonGroupID)
+		{
+			for (const FMeshTriangle& MeshTriangle : MeshPolygon.Triangles)
+			{
+				OutTriangles.Add(MeshTriangle);
+				TriangleCount++;
+			}
+		}
+	}
+	return TriangleCount;
+}
+
+bool IsOrphanedVertex(const UMeshDescription *MeshDescription, const FVertexID VertexID)
 {
 	const FVertexInstanceArray& VertexInstances = MeshDescription->VertexInstances();
 	for (const FVertexInstanceID VertexInstanceID : MeshDescription->GetVertex(VertexID).VertexInstanceIDs)
@@ -236,7 +152,7 @@ bool FStaticMeshBuilder::IsOrphanedVertex(const UMeshDescription *MeshDescriptio
 	return true;
 }
 
-void FStaticMeshBuilder::UpdateBounds(UStaticMesh* StaticMesh)
+void UpdateBounds(UStaticMesh* StaticMesh)
 {
 	const UMeshDescription* MeshDescription = StaticMesh->GetMeshDescription();
 	check(MeshDescription != nullptr);
@@ -284,7 +200,7 @@ void FStaticMeshBuilder::UpdateBounds(UStaticMesh* StaticMesh)
 }
 
 
-void FStaticMeshBuilder::UpdateCollision(UStaticMesh *StaticMesh)
+void UpdateCollision(UStaticMesh *StaticMesh)
 {
 	// @todo mesheditor collision: We're wiping the existing simplified collision and generating a simple bounding
 	// box collision, since that's the best we can do without impacting performance.  We always using visibility (complex)
@@ -341,3 +257,180 @@ void FStaticMeshBuilder::UpdateCollision(UStaticMesh *StaticMesh)
 		}
 	}
 }
+
+void BuildVertexBuffer(const UMeshDescription* MeshDescription, FStaticMeshLODResources& StaticMeshLOD, const FMeshBuildSettings& LODBuildSettings, TArray< FStaticMeshBuildVertex >& StaticMeshBuildVertices)
+{
+	const FVertexArray& Vertices = MeshDescription->Vertices();
+	const FVertexInstanceArray& VertexInstances = MeshDescription->VertexInstances();
+
+	// set up vertex buffer elements
+	StaticMeshBuildVertices.SetNum(VertexInstances.GetArraySize());
+	bool bHasColor = false;
+	for (const FVertexInstanceID VertexInstanceID : VertexInstances.GetElementIDs())
+	{
+		const FMeshVertexInstance& VertexInstance = VertexInstances[VertexInstanceID];
+
+		if (VertexInstance.Color != FColor::White)
+		{
+			bHasColor = true;
+		}
+
+		FStaticMeshBuildVertex& StaticMeshVertex = StaticMeshBuildVertices[VertexInstanceID.GetValue()];
+
+		StaticMeshVertex.Position = Vertices[VertexInstance.VertexID].VertexPosition * LODBuildSettings.BuildScale3D;
+		StaticMeshVertex.TangentX = VertexInstance.Tangent;
+		StaticMeshVertex.TangentY = FVector::CrossProduct(VertexInstance.Normal, VertexInstance.Tangent).GetSafeNormal() * VertexInstance.BinormalSign;
+		StaticMeshVertex.TangentZ = VertexInstance.Normal;
+		StaticMeshVertex.Color = VertexInstance.Color.ToFColor(true);
+		for (int32 UVIndex = 0; UVIndex < VertexInstance.VertexUVs.Num(); ++UVIndex)
+		{
+			StaticMeshVertex.UVs[UVIndex] = VertexInstance.VertexUVs[UVIndex];
+		}
+	}
+
+	StaticMeshLOD.PositionVertexBuffer.Init(StaticMeshBuildVertices);
+	StaticMeshLOD.VertexBuffer.SetUseHighPrecisionTangentBasis(LODBuildSettings.bUseHighPrecisionTangentBasis);
+	StaticMeshLOD.VertexBuffer.SetUseFullPrecisionUVs(LODBuildSettings.bUseFullPrecisionUVs);
+	StaticMeshLOD.VertexBuffer.Init(StaticMeshBuildVertices, 1); //TODO: use the real texture coordinnate count there is only one currently in the meshdescription
+	if (bHasColor)
+	{
+		StaticMeshLOD.ColorVertexBuffer.Init(StaticMeshBuildVertices);
+	}
+	else
+	{
+		StaticMeshLOD.ColorVertexBuffer.InitFromSingleColor(FColor::White, StaticMeshBuildVertices.Num());
+	}
+}
+
+void BuildIndexBuffer(UStaticMesh *StaticMesh, const UMeshDescription* MeshDescription, TArray< uint32 >& IndexBuffer, FStaticMeshLODResources& StaticMeshLOD)
+{
+	const FPolygonGroupArray& PolygonGroups = MeshDescription->PolygonGroups();
+	// Set up index buffer
+	for (const FPolygonGroupID& PolygonGroupID : PolygonGroups.GetElementIDs())
+	{
+		TArray<FMeshTriangle> Triangles;
+		int32 SectionTriangleCount = GetPolygonGroupTriangles(MeshDescription, Triangles, PolygonGroupID);
+
+		const FMeshPolygonGroup& PolygonGroup = PolygonGroups[PolygonGroupID];
+
+		// Create new rendering section
+		StaticMeshLOD.Sections.Add(FStaticMeshSection());
+		FStaticMeshSection& StaticMeshSection = StaticMeshLOD.Sections.Last();
+
+		StaticMeshSection.FirstIndex = IndexBuffer.Num();
+		StaticMeshSection.NumTriangles = SectionTriangleCount;
+
+		const int32 MaterialIndex = StaticMesh->GetMaterialIndex(PolygonGroup.ImportedMaterialSlotName);
+		check(MaterialIndex != INDEX_NONE);
+		check(StaticMesh->StaticMaterials[MaterialIndex].MaterialInterface->GetPathName() == PolygonGroup.MaterialAsset.ToString());
+		StaticMeshSection.MaterialIndex = MaterialIndex;
+		StaticMeshSection.bEnableCollision = PolygonGroup.bEnableCollision;
+		StaticMeshSection.bCastShadow = PolygonGroup.bCastShadow;
+
+		if (Triangles.Num() > 0)
+		{
+			IndexBuffer.Reserve(IndexBuffer.Num() + Triangles.Num() * 3);
+			uint32 MinIndex = TNumericLimits< uint32 >::Max();
+			uint32 MaxIndex = TNumericLimits< uint32 >::Min();
+			for (int32 TriangleIndex = 0; TriangleIndex < Triangles.Num(); ++TriangleIndex)
+			{
+				const FMeshTriangle& Triangle = Triangles[TriangleIndex];
+				for (int32 TriVert = 0; TriVert < 3; ++TriVert)
+				{
+					const uint32 RenderingVertexIndex = Triangle.GetVertexInstanceID(TriVert).GetValue();
+					IndexBuffer.Add(RenderingVertexIndex);
+					MinIndex = FMath::Min(MinIndex, RenderingVertexIndex);
+					MaxIndex = FMath::Max(MaxIndex, RenderingVertexIndex);
+				}
+			}
+			StaticMeshSection.MinVertexIndex = MinIndex;
+			StaticMeshSection.MaxVertexIndex = MaxIndex;
+		}
+		else
+		{
+			// No triangles in this section
+			StaticMeshSection.MinVertexIndex = 0;
+			StaticMeshSection.MaxVertexIndex = 0;
+		}
+	}
+}
+
+void BuildAllBufferOptimizations(FStaticMeshLODResources& StaticMeshLOD, const FMeshBuildSettings& LODBuildSettings, TArray< uint32 >& IndexBuffer, bool bNeeds32BitIndices, TArray< FStaticMeshBuildVertex >& StaticMeshBuildVertices)
+{
+	const EIndexBufferStride::Type IndexBufferStride = bNeeds32BitIndices ? EIndexBufferStride::Force32Bit : EIndexBufferStride::Force16Bit;
+
+	// Build the reversed index buffer.
+	if (LODBuildSettings.bBuildReversedIndexBuffer)
+	{
+		TArray<uint32> InversedIndices;
+		const int32 IndexCount = IndexBuffer.Num();
+		InversedIndices.AddUninitialized(IndexCount);
+
+		for (int32 SectionIndex = 0; SectionIndex < StaticMeshLOD.Sections.Num(); ++SectionIndex)
+		{
+			const FStaticMeshSection& SectionInfo = StaticMeshLOD.Sections[SectionIndex];
+			const int32 SectionIndexCount = SectionInfo.NumTriangles * 3;
+
+			for (int32 i = 0; i < SectionIndexCount; ++i)
+			{
+				InversedIndices[SectionInfo.FirstIndex + i] = IndexBuffer[SectionInfo.FirstIndex + SectionIndexCount - 1 - i];
+			}
+		}
+		StaticMeshLOD.ReversedIndexBuffer.SetIndices(InversedIndices, IndexBufferStride);
+	}
+
+	// Build the depth-only index buffer.
+	TArray<uint32> DepthOnlyIndices;
+	{
+		BuildOptimizationHelper::BuildDepthOnlyIndexBuffer(
+			DepthOnlyIndices,
+			StaticMeshBuildVertices,
+			IndexBuffer,
+			StaticMeshLOD.Sections
+		);
+		StaticMeshLOD.DepthOnlyIndexBuffer.SetIndices(DepthOnlyIndices, IndexBufferStride);
+	}
+
+	// Build the inversed depth only index buffer.
+	if (LODBuildSettings.bBuildReversedIndexBuffer)
+	{
+		TArray<uint32> ReversedDepthOnlyIndices;
+		const int32 IndexCount = DepthOnlyIndices.Num();
+		ReversedDepthOnlyIndices.AddUninitialized(IndexCount);
+		for (int32 i = 0; i < IndexCount; ++i)
+		{
+			ReversedDepthOnlyIndices[i] = DepthOnlyIndices[IndexCount - 1 - i];
+		}
+		StaticMeshLOD.ReversedDepthOnlyIndexBuffer.SetIndices(ReversedDepthOnlyIndices, IndexBufferStride);
+	}
+
+	// Build a list of wireframe edges in the static mesh.
+	{
+		TArray<BuildOptimizationHelper::FMeshEdge> Edges;
+		TArray<uint32> WireframeIndices;
+
+		BuildOptimizationHelper::FStaticMeshEdgeBuilder(IndexBuffer, StaticMeshBuildVertices, Edges).FindEdges();
+		WireframeIndices.Empty(2 * Edges.Num());
+		for (int32 EdgeIndex = 0; EdgeIndex < Edges.Num(); EdgeIndex++)
+		{
+			BuildOptimizationHelper::FMeshEdge&	Edge = Edges[EdgeIndex];
+			WireframeIndices.Add(Edge.Vertices[0]);
+			WireframeIndices.Add(Edge.Vertices[1]);
+		}
+		StaticMeshLOD.WireframeIndexBuffer.SetIndices(WireframeIndices, IndexBufferStride);
+	}
+
+	// Build the adjacency index buffer used for tessellation.
+	if (LODBuildSettings.bBuildAdjacencyBuffer)
+	{
+		TArray<uint32> AdjacencyIndices;
+
+		BuildOptimizationHelper::BuildStaticAdjacencyIndexBuffer(
+			StaticMeshLOD.PositionVertexBuffer,
+			StaticMeshLOD.VertexBuffer,
+			IndexBuffer,
+			AdjacencyIndices
+		);
+		StaticMeshLOD.AdjacencyIndexBuffer.SetIndices(AdjacencyIndices, IndexBufferStride);
+	}
+}
\ No newline at end of file
diff --git a/Engine/Source/Runtime/MeshBuilder/Public/IMeshBuilderModule.h b/Engine/Source/Runtime/MeshBuilder/Public/IMeshBuilderModule.h
--- a/Engine/Source/Runtime/MeshBuilder/Public/IMeshBuilderModule.h
+++ b/Engine/Source/Runtime/MeshBuilder/Public/IMeshBuilderModule.h
@@ -4,6 +4,7 @@
 
 #include "CoreMinimal.h"
 #include "Modules/ModuleInterface.h"
+#include "Modules/ModuleManager.h"
 
 class IMeshBuilderModule : public IModuleInterface
 {
diff --git a/Engine/Source/Runtime/MeshBuilder/Public/StaticMeshBuilder.h b/Engine/Source/Runtime/MeshBuilder/Public/StaticMeshBuilder.h
--- a/Engine/Source/Runtime/MeshBuilder/Public/StaticMeshBuilder.h
+++ b/Engine/Source/Runtime/MeshBuilder/Public/StaticMeshBuilder.h
@@ -4,32 +4,20 @@
 
 #include "CoreMinimal.h"
 #include "MeshBuilder.h"
-#include "MeshDescription.h"
 
 
 class MESHBUILDER_API FStaticMeshBuilder : public FMeshBuilder
 {
 public:
-	static FStaticMeshBuilder& Get()
-	{
-		static FStaticMeshBuilder Instance;
-		return Instance;
-	}
+	FStaticMeshBuilder();
 
 	virtual bool Build(class UStaticMesh* StaticMesh) override;
 	virtual ~FStaticMeshBuilder() {}
 
 private:
-	FStaticMeshBuilder();
 
 	void OnBuildRenderMeshStart(class UStaticMesh* StaticMesh, const bool bInvalidateLighting);
 	void OnBuildRenderMeshFinish(class UStaticMesh* StaticMesh, const bool bRebuildBoundsAndCollision);
-	
-	int32 GetPolygonGroupTriangles(const class UMeshDescription *MeshDescription, TArray<FMeshTriangle>& OutTriangles, const FPolygonGroupID& PolygonGroupID);
-
-	bool IsOrphanedVertex(const class UMeshDescription *MeshDescription, const FVertexID VertexID) const;
-	void UpdateBounds(class UStaticMesh* StaticMesh);
-	void UpdateCollision(UStaticMesh *StaticMesh);
 
 	/** Used to refresh all components in the scene that may be using a mesh we're editing */
 	TSharedPtr<class FStaticMeshComponentRecreateRenderStateContext> RecreateRenderStateContext;
