Add the Render data wedge map

diff --git a/Engine/Source/Runtime/MeshBuilder/Private/StaticMeshBuilder.cpp b/Engine/Source/Runtime/MeshBuilder/Private/StaticMeshBuilder.cpp
--- a/Engine/Source/Runtime/MeshBuilder/Private/StaticMeshBuilder.cpp
+++ b/Engine/Source/Runtime/MeshBuilder/Private/StaticMeshBuilder.cpp
@@ -15,7 +15,7 @@ bool IsOrphanedVertex(const class UMeshDescription *MeshDescription, const FVert
 void UpdateBounds(class UStaticMesh* StaticMesh);
 void UpdateCollision(UStaticMesh *StaticMesh);
 void BuildVertexBuffer(const class UMeshDescription* MeshDescription, struct FStaticMeshLODResources& StaticMeshLOD, const struct FMeshBuildSettings& LODBuildSettings, TArray< FStaticMeshBuildVertex >& StaticMeshBuildVertices);
-void BuildIndexBuffer(class UStaticMesh *StaticMesh, const class UMeshDescription* MeshDescription, TArray< uint32 >& IndexBuffer, struct FStaticMeshLODResources& StaticMeshLOD);
+void BuildIndexBuffer(class UStaticMesh *StaticMesh, const class UMeshDescription* MeshDescription, TArray< uint32 >& IndexBuffer, struct FStaticMeshLODResources& StaticMeshLOD, TArray<int32>& OutWedgeMap);
 void BuildAllBufferOptimizations(struct FStaticMeshLODResources& StaticMeshLOD, const struct FMeshBuildSettings& LODBuildSettings, TArray< uint32 >& IndexBuffer, bool bNeeds32BitIndices, TArray< FStaticMeshBuildVertex >& StaticMeshBuildVertices);
 //////////////////////////////////////////////////////////////////////////
 
@@ -60,8 +60,11 @@ bool FStaticMeshBuilder::Build(UStaticMesh* StaticMesh)
 
 		BuildVertexBuffer(MeshDescription, StaticMeshLOD, LODBuildSettings, StaticMeshBuildVertices);
 
-		BuildIndexBuffer(StaticMesh, MeshDescription, IndexBuffer, StaticMeshLOD);
-		
+		//Render data Wedge map is only set for LOD 0???
+		TArray<int32> TempWedgeMap;
+		TArray<int32> &WedgeMap = (LodIndex == 0) ? StaticMeshRenderData.WedgeMap : TempWedgeMap;
+		BuildIndexBuffer(StaticMesh, MeshDescription, IndexBuffer, StaticMeshLOD, WedgeMap);
+
 		// Figure out which index buffer stride we need
 		bool bNeeds32BitIndices = false;
 		for (const FStaticMeshSection& StaticMeshSection : StaticMeshLOD.Sections)
@@ -302,8 +305,10 @@ void BuildVertexBuffer(const UMeshDescription* MeshDescription, FStaticMeshLODRe
 	}
 }
 
-void BuildIndexBuffer(UStaticMesh *StaticMesh, const UMeshDescription* MeshDescription, TArray< uint32 >& IndexBuffer, FStaticMeshLODResources& StaticMeshLOD)
+void BuildIndexBuffer(UStaticMesh *StaticMesh, const UMeshDescription* MeshDescription, TArray< uint32 >& IndexBuffer, FStaticMeshLODResources& StaticMeshLOD, TArray<int32>& OutWedgeMap)
 {
+	OutWedgeMap.Reset();
+
 	const FPolygonGroupArray& PolygonGroups = MeshDescription->PolygonGroups();
 	// Set up index buffer
 	for (const FPolygonGroupID& PolygonGroupID : PolygonGroups.GetElementIDs())
@@ -329,7 +334,9 @@ void BuildIndexBuffer(UStaticMesh *StaticMesh, const UMeshDescription* MeshDescr
 
 		if (Triangles.Num() > 0)
 		{
-			IndexBuffer.Reserve(IndexBuffer.Num() + Triangles.Num() * 3);
+			int32 ReseveSize = IndexBuffer.Num() + Triangles.Num() * 3;
+			IndexBuffer.Reserve(ReseveSize);
+			OutWedgeMap.Reserve(ReseveSize);
 			uint32 MinIndex = TNumericLimits< uint32 >::Max();
 			uint32 MaxIndex = TNumericLimits< uint32 >::Min();
 			for (int32 TriangleIndex = 0; TriangleIndex < Triangles.Num(); ++TriangleIndex)
@@ -339,6 +346,7 @@ void BuildIndexBuffer(UStaticMesh *StaticMesh, const UMeshDescription* MeshDescr
 				{
 					const uint32 RenderingVertexIndex = Triangle.GetVertexInstanceID(TriVert).GetValue();
 					IndexBuffer.Add(RenderingVertexIndex);
+					OutWedgeMap.Add(RenderingVertexIndex);
 					MinIndex = FMath::Min(MinIndex, RenderingVertexIndex);
 					MaxIndex = FMath::Max(MaxIndex, RenderingVertexIndex);
 				}
