Use the current static build code (ddc, etc) to build render data from mesh description.
Change the name of the fbx import option to import mesh description
Add missing indexes buffer GPU optimisation

diff --git a/Engine/Config/BaseEditorPerProjectUserSettings.ini b/Engine/Config/BaseEditorPerProjectUserSettings.ini
--- a/Engine/Config/BaseEditorPerProjectUserSettings.ini
+++ b/Engine/Config/BaseEditorPerProjectUserSettings.ini
@@ -575,7 +575,7 @@ BackupIntervalInMinutes=5
 
 
 [/Script/UnrealEd.FbxImportUI]
-bImportEditableMesh=False
+bImportMeshDescription=False
 bImportMaterials=True
 bAutoCreateGroups=True
 bInvertNormalMaps=False
diff --git a/Engine/Source/Editor/UnrealEd/Classes/Factories/FbxImportUI.h b/Engine/Source/Editor/UnrealEd/Classes/Factories/FbxImportUI.h
--- a/Engine/Source/Editor/UnrealEd/Classes/Factories/FbxImportUI.h
+++ b/Engine/Source/Editor/UnrealEd/Classes/Factories/FbxImportUI.h
@@ -52,7 +52,7 @@ public:
 
 	/** Use the string in "Name" field as full name of mesh. The option only works when the scene contains one mesh. */
 	UPROPERTY(EditAnywhere, AdvancedDisplay, config, Category = Miscellaneous, meta = (OBJRestrict = "true"))
-	uint32 bImportEditableMesh:1;
+	uint32 bImportMeshDescription:1;
 
 	/** Whether to import the incoming FBX as a skeletal object */
 	UPROPERTY(EditAnywhere, Category = Mesh, meta = (ImportType = "StaticMesh|SkeletalMesh", DisplayName="Skeletal Mesh"))
diff --git a/Engine/Source/Editor/UnrealEd/Private/Fbx/FbxMainImport.cpp b/Engine/Source/Editor/UnrealEd/Private/Fbx/FbxMainImport.cpp
--- a/Engine/Source/Editor/UnrealEd/Private/Fbx/FbxMainImport.cpp
+++ b/Engine/Source/Editor/UnrealEd/Private/Fbx/FbxMainImport.cpp
@@ -247,7 +247,7 @@ void ApplyImportUIToImportOptions(UFbxImportUI* ImportUI, FBXImportOptions& InOu
 		InOutImportOptions.BaseEmissiveColorName = ImportUI->TextureImportData->BaseEmissiveColorName;
 	}
 	InOutImportOptions.bImportTextures = ImportUI->bImportTextures;
-	InOutImportOptions.bImportEditableMesh = ImportUI->bImportEditableMesh;
+	InOutImportOptions.bImportMeshDescription = ImportUI->bImportMeshDescription;
 	InOutImportOptions.bUsedAsFullName = ImportUI->bOverrideFullName;
 	InOutImportOptions.bImportAnimations = ImportUI->bImportAnimations;
 	InOutImportOptions.SkeletonForAnimation = ImportUI->Skeleton;
diff --git a/Engine/Source/Editor/UnrealEd/Private/Fbx/FbxStaticMeshImport.cpp b/Engine/Source/Editor/UnrealEd/Private/Fbx/FbxStaticMeshImport.cpp
--- a/Engine/Source/Editor/UnrealEd/Private/Fbx/FbxStaticMeshImport.cpp
+++ b/Engine/Source/Editor/UnrealEd/Private/Fbx/FbxStaticMeshImport.cpp
@@ -266,7 +266,7 @@ bool UnFbx::FFbxImporter::BuildStaticMeshFromGeometry(FbxNode* Node, UStaticMesh
 	FStaticMeshSourceModel& SrcModel = StaticMesh->SourceModels[LODIndex];
 	
 	UMeshDescription *MeshDescription = StaticMesh->GetMeshDescription(LODIndex);
-	if (ImportOptions->bImportEditableMesh)
+	if (ImportOptions->bImportMeshDescription)
 	{
 		//The mesh description should have been created before calling BuildStaticMeshFromGeometry
 		check(MeshDescription);
@@ -524,17 +524,6 @@ bool UnFbx::FFbxImporter::BuildStaticMeshFromGeometry(FbxNode* Node, UStaticMesh
 		int32 VertexInstanceOffset = MeshDescription->VertexInstances().Num();
 		int32 TriangleOffset = MeshDescription->Polygons().Num();
 
-		if (StaticMesh->RenderData == nullptr)
-		{
-			StaticMesh->CacheDerivedData();
-		}
-		//Add missing LODs render data resources
-		while (StaticMesh->RenderData->LODResources.Num() <= LODIndex)
-		{
-			FStaticMeshLODResources *DummyStaticMeshLodResources = new FStaticMeshLODResources();
-			StaticMesh->RenderData->LODResources.Add(DummyStaticMeshLodResources);
-		}
-
 		int32 MaxMaterialIndex = 0;
 		TMap<int32, FPolygonGroupID> PolygonGroupMapping;
 		int32 NumUVs = FBXUVs.UniqueUVCount;
@@ -1563,7 +1552,7 @@ UStaticMesh* UnFbx::FFbxImporter::ImportStaticMeshAsSingle(UObject* InParent, TA
 		StaticMesh = NewObject<UStaticMesh>(Package, FName(*MeshName), Flags | RF_Public);
 	}
 
-	if (ImportOptions->bImportEditableMesh)
+	if (ImportOptions->bImportMeshDescription)
 	{
 		MeshDescription = StaticMesh->GetMeshDescription(LODIndex);
 		if (MeshDescription == nullptr)
@@ -1593,6 +1582,8 @@ UStaticMesh* UnFbx::FFbxImporter::ImportStaticMeshAsSingle(UObject* InParent, TA
 		FRawMesh EmptyRawMesh;
 		SrcModel.RawMeshBulkData->SaveRawMesh( EmptyRawMesh );
 	}
+
+	StaticMesh->ClearOriginalMeshDescription(LODIndex);
 	
 	// make sure it has a new lighting guid
 	StaticMesh->LightingGuid = FGuid::NewGuid();
@@ -1945,6 +1936,8 @@ UStaticMesh* UnFbx::FFbxImporter::ImportStaticMeshAsSingle(UObject* InParent, TA
 					}
 				}
 			}
+			//Set the original mesh description to be able to do non destructive reduce
+			StaticMesh->SetOriginalMeshDescription(LODIndex, MeshDescription);
 		}
 
 		// Setup default LOD settings based on the selected LOD group.
@@ -2094,19 +2087,11 @@ void UnFbx::FFbxImporter::PostImportStaticMesh(UStaticMesh* StaticMesh, TArray<F
 				}
 				bool bOriginalGenerateMeshDistanceField = StaticMesh->bGenerateMeshDistanceField;
 				StaticMesh->bGenerateMeshDistanceField = false;
-				
-				if (ImportOptions->bImportEditableMesh)
-				{
-					IMeshBuilderModule& MeshBuilder = FModuleManager::Get().LoadModuleChecked<IMeshBuilderModule>(TEXT("MeshBuilder"));
-					MeshBuilder.BuildMesh(StaticMesh);
-				}
-				else
+
+				StaticMesh->Build(false, &BuildErrors);
+				for (FText& Error : BuildErrors)
 				{
-					StaticMesh->Build(false, &BuildErrors);
-					for (FText& Error : BuildErrors)
-					{
-						AddTokenizedErrorMessage(FTokenizedMessage::Create(EMessageSeverity::Warning, Error), FFbxErrors::StaticMesh_BuildError);
-					}
+					AddTokenizedErrorMessage(FTokenizedMessage::Create(EMessageSeverity::Warning, Error), FFbxErrors::StaticMesh_BuildError);
 				}
 
 				StaticMesh->bGenerateMeshDistanceField = bOriginalGenerateMeshDistanceField;
@@ -2140,18 +2125,10 @@ void UnFbx::FFbxImporter::PostImportStaticMesh(UStaticMesh* StaticMesh, TArray<F
 			}
 		}
 	}
-	if (ImportOptions->bImportEditableMesh)
+	StaticMesh->Build(false, &BuildErrors);
+	for (FText& Error : BuildErrors)
 	{
-		IMeshBuilderModule& MeshBuilder = FModuleManager::Get().LoadModuleChecked<IMeshBuilderModule>(TEXT("MeshBuilder"));
-		MeshBuilder.BuildMesh(StaticMesh);
-	}
-	else
-	{
-		StaticMesh->Build(false, &BuildErrors);
-		for (FText& Error : BuildErrors)
-		{
-			AddTokenizedErrorMessage(FTokenizedMessage::Create(EMessageSeverity::Warning, Error), FFbxErrors::StaticMesh_BuildError);
-		}
+		AddTokenizedErrorMessage(FTokenizedMessage::Create(EMessageSeverity::Warning, Error), FFbxErrors::StaticMesh_BuildError);
 	}
 
 	//Set the specified LOD distances for every LODs we have to do this after the build in case there is a specified Lod Group
diff --git a/Engine/Source/Editor/UnrealEd/Public/FbxImporter.h b/Engine/Source/Editor/UnrealEd/Public/FbxImporter.h
--- a/Engine/Source/Editor/UnrealEd/Public/FbxImporter.h
+++ b/Engine/Source/Editor/UnrealEd/Public/FbxImporter.h
@@ -107,7 +107,7 @@ namespace UnFbx
 struct FBXImportOptions
 {
 	// General options
-	bool bImportEditableMesh;
+	bool bImportMeshDescription;
 	bool bImportScene;
 	bool bImportMaterials;
 	bool bResetMaterialSlots;
diff --git a/Engine/Source/Runtime/Engine/Classes/Engine/StaticMesh.h b/Engine/Source/Runtime/Engine/Classes/Engine/StaticMesh.h
--- a/Engine/Source/Runtime/Engine/Classes/Engine/StaticMesh.h
+++ b/Engine/Source/Runtime/Engine/Classes/Engine/StaticMesh.h
@@ -629,6 +629,19 @@ public:
 	ENGINE_API class UMeshDescription* GetMeshDescription(int32 LodIndex=0) const;
 	ENGINE_API int32 GetMeshDescriptionCount() const;
 
+#if WITH_EDITORONLY_DATA
+	/**
+	* The original mesh description we use to do non destructive reduce of the mesh, we set the original value when we create the mesh (procedurally or at import/re-import...).
+	*/
+	UPROPERTY()
+	TArray<class UMeshDescription*> OriginalMeshDescriptions;
+
+	ENGINE_API class UMeshDescription* GetOriginalMeshDescription(int32 LodIndex = 0) const;
+	ENGINE_API void SetOriginalMeshDescription(int32 LodIndex, class UMeshDescription* MeshDescription);
+	ENGINE_API int32 GetOriginalMeshDescriptionCount() const;
+	ENGINE_API void ClearOriginalMeshDescription(int32 LodIndex);
+#endif
+
 	/** Pre-build navigation collision */
 	UPROPERTY(VisibleAnywhere, transient, duplicatetransient, Instanced, Category = Navigation)
 	class UNavCollision* NavCollision;
diff --git a/Engine/Source/Runtime/Engine/Private/StaticMesh.cpp b/Engine/Source/Runtime/Engine/Private/StaticMesh.cpp
--- a/Engine/Source/Runtime/Engine/Private/StaticMesh.cpp
+++ b/Engine/Source/Runtime/Engine/Private/StaticMesh.cpp
@@ -42,6 +42,7 @@
 #include "MeshDescription.h"
 #include "MeshUtilities.h"
 #include "DerivedDataCacheInterface.h"
+#include "IMeshBuilderModule.h"
 #endif // #if WITH_EDITOR
 
 #include "Engine/StaticMeshSocket.h"
@@ -1141,7 +1142,14 @@ static FString BuildStaticMeshDerivedDataKey(UStaticMesh* Mesh, const FStaticMes
 	for (int32 LODIndex = 0; LODIndex < NumLODs; ++LODIndex)
 	{
 		FStaticMeshSourceModel& SrcModel = Mesh->SourceModels[LODIndex];
-		KeySuffix += SrcModel.RawMeshBulkData->GetIdString();
+		if (Mesh->GetOriginalMeshDescriptionCount() > 0)
+		{
+			KeySuffix += (Mesh->GetOriginalMeshDescriptionCount() > LODIndex && Mesh->GetOriginalMeshDescription(LODIndex) != nullptr) ? Mesh->GetOriginalMeshDescription(LODIndex)->GetIdString() : TEXT("NoOriginalMeshDescriptionForLod") + FString::FromInt(LODIndex);
+		}
+		else
+		{
+			KeySuffix += SrcModel.RawMeshBulkData->GetIdString();
+		}
 
 		// Serialize the build and reduction settings into a temporary array. The archive
 		// is flagged as persistent so that machines of different endianness produce
@@ -1285,11 +1293,19 @@ void FStaticMeshRenderData::Cache(UStaticMesh* Owner, const FStaticMeshLODSettin
 			Args.Add(TEXT("StaticMeshName"), FText::FromString( Owner->GetName() ) );
 			FStaticMeshStatusMessageContext StatusContext( FText::Format( NSLOCTEXT("Engine", "BuildingStaticMeshStatus", "Building static mesh {StaticMeshName}..."), Args ) );
 
-			IMeshUtilities& MeshUtilities = FModuleManager::Get().LoadModuleChecked<IMeshUtilities>(TEXT("MeshUtilities"));
-			if (!MeshUtilities.BuildStaticMesh(*this, Owner->SourceModels, LODGroup, Owner->LightmapUVVersion, Owner->ImportVersion))
+			if (Owner->GetMeshDescriptionCount() > 0)
 			{
-				UE_LOG(LogStaticMesh, Error, TEXT("Failed to build static mesh. See previous line(s) for details."));
-				return;
+				IMeshBuilderModule& MeshBuilderModule = FModuleManager::Get().LoadModuleChecked<IMeshBuilderModule>(TEXT("MeshBuilder"));
+				MeshBuilderModule.BuildMesh(Owner);
+			}
+			else
+			{
+				IMeshUtilities& MeshUtilities = FModuleManager::Get().LoadModuleChecked<IMeshUtilities>(TEXT("MeshUtilities"));
+				if (!MeshUtilities.BuildStaticMesh(*this, Owner->SourceModels, LODGroup, Owner->LightmapUVVersion, Owner->ImportVersion))
+				{
+					UE_LOG(LogStaticMesh, Error, TEXT("Failed to build static mesh. See previous line(s) for details."));
+					return;
+				}
 			}
 
 			ComputeUVDensities();
@@ -2197,11 +2213,49 @@ UMeshDescription* UStaticMesh::GetMeshDescription(int32 LodIndex) const
 	return MeshDescriptions[LodIndex];
 }
 
-ENGINE_API int32 UStaticMesh::GetMeshDescriptionCount() const
+int32 UStaticMesh::GetMeshDescriptionCount() const
 {
 	return MeshDescriptions.Num();
 }
 
+#if WITH_EDITORONLY_DATA
+
+UMeshDescription* UStaticMesh::GetOriginalMeshDescription(int32 LodIndex) const
+{
+	if (!OriginalMeshDescriptions.IsValidIndex(LodIndex))
+	{
+		return nullptr;
+	}
+	return OriginalMeshDescriptions[LodIndex];
+}
+
+void UStaticMesh::SetOriginalMeshDescription(int32 LodIndex, class UMeshDescription* MeshDescription)
+{
+	if (!OriginalMeshDescriptions.IsValidIndex(LodIndex))
+	{
+		//Add nullptr missing entries
+		OriginalMeshDescriptions.AddZeroed(LodIndex - OriginalMeshDescriptions.Num() + 1);
+	}
+	//Set the original mesh description
+	OriginalMeshDescriptions[LodIndex] = MeshDescription;
+}
+
+int32 UStaticMesh::GetOriginalMeshDescriptionCount() const
+{
+	return OriginalMeshDescriptions.Num();
+}
+
+void UStaticMesh::ClearOriginalMeshDescription(int32 LodIndex)
+{
+	if (!OriginalMeshDescriptions.IsValidIndex(LodIndex))
+	{
+		return;
+	}
+	OriginalMeshDescriptions[LodIndex] = nullptr;
+}
+
+#endif
+
 void UStaticMesh::CacheDerivedData()
 {
 	// Cache derived data for the running platform.
diff --git a/Engine/Source/Runtime/Engine/Private/StaticMeshBuild.cpp b/Engine/Source/Runtime/Engine/Private/StaticMeshBuild.cpp
--- a/Engine/Source/Runtime/Engine/Private/StaticMeshBuild.cpp
+++ b/Engine/Source/Runtime/Engine/Private/StaticMeshBuild.cpp
@@ -76,7 +76,7 @@ void UStaticMesh::Build(bool bSilent, TArray<FText>* OutErrors)
 		return;
 
 	// If we're controlled by an editable mesh do not build. The editable mesh will build us
-	if (EditableMesh || MeshDescriptions.Num() > 0)
+	if (EditableMesh)
 	{
 		return;
 	}
diff --git a/Engine/Source/Runtime/MeshBuilder/MeshBuilder.Build.cs b/Engine/Source/Runtime/MeshBuilder/MeshBuilder.Build.cs
--- a/Engine/Source/Runtime/MeshBuilder/MeshBuilder.Build.cs
+++ b/Engine/Source/Runtime/MeshBuilder/MeshBuilder.Build.cs
@@ -21,6 +21,7 @@ namespace UnrealBuildTool.Rules
 			);
 
             AddEngineThirdPartyPrivateStaticDependencies(Target, "nvTriStrip");
+            AddEngineThirdPartyPrivateStaticDependencies(Target, "ForsythTriOptimizer");
             AddEngineThirdPartyPrivateStaticDependencies(Target, "nvTessLib");
         }
 	}
diff --git a/Engine/Source/Runtime/MeshBuilder/Private/BuildOptimizationHelper.h b/Engine/Source/Runtime/MeshBuilder/Private/BuildOptimizationHelper.h
--- a/Engine/Source/Runtime/MeshBuilder/Private/BuildOptimizationHelper.h
+++ b/Engine/Source/Runtime/MeshBuilder/Private/BuildOptimizationHelper.h
@@ -4,10 +4,189 @@
 #include "CoreMinimal.h"
 #include "CoreTypes.h"
 #include "NvTriStrip.h"
+#include "forsythtriangleorderoptimizer.h"
 #include "nvtess.h"
 #include "Components.h"
 #include "StaticMeshResources.h"
 
+namespace BuildOptimizationPrivate
+{
+	// CVars
+	static TAutoConsoleVariable<int32> CVarTriangleOrderOptimization(
+		TEXT("r.TriangleOrderOptimization"),
+		1,
+		TEXT("Controls the algorithm to use when optimizing the triangle order for the post-transform cache.\n")
+		TEXT("0: Use NVTriStrip (slower)\n")
+		TEXT("1: Use Forsyth algorithm (fastest)(default)")
+		TEXT("2: No triangle order optimization. (least efficient, debugging purposes only)"),
+		ECVF_Default);
+
+	static TAutoConsoleVariable<int32> CVarSupportDepthOnlyIndexBuffers(
+		TEXT("r.SupportDepthOnlyIndexBuffers"),
+		1,
+		TEXT("Enables depth-only index buffers. Saves a little time at the expense of doubling the size of index buffers."),
+		ECVF_ReadOnly | ECVF_RenderThreadSafe);
+
+	static TAutoConsoleVariable<int32> CVarSupportReversedIndexBuffers(
+		TEXT("r.SupportReversedIndexBuffers"),
+		1,
+		TEXT("Enables reversed index buffers. Saves a little time at the expense of doubling the size of index buffers."),
+		ECVF_ReadOnly | ECVF_RenderThreadSafe);
+
+	/*------------------------------------------------------------------------------
+	NVTriStrip for cache optimizing index buffers.
+	------------------------------------------------------------------------------*/
+
+	namespace NvTriStripHelper
+	{
+		/**
+		* Converts 16 bit indices to 32 bit prior to passing them into the real GenerateStrips util method
+		*/
+		void GenerateStrips(
+			const uint8* Indices,
+			bool Is32Bit,
+			const uint32 NumIndices,
+			PrimitiveGroup** PrimGroups,
+			uint32* NumGroups
+		)
+		{
+			if (Is32Bit)
+			{
+				GenerateStrips((uint32*)Indices, NumIndices, PrimGroups, NumGroups);
+			}
+			else
+			{
+				// convert to 32 bit
+				uint32 Idx;
+				TArray<uint32> NewIndices;
+				NewIndices.AddUninitialized(NumIndices);
+				for (Idx = 0; Idx < NumIndices; ++Idx)
+				{
+					NewIndices[Idx] = ((uint16*)Indices)[Idx];
+				}
+				GenerateStrips(NewIndices.GetData(), NumIndices, PrimGroups, NumGroups);
+			}
+
+		}
+
+		/**
+		* Orders a triangle list for better vertex cache coherency.
+		*
+		* *** WARNING: This is safe to call for multiple threads IF AND ONLY IF all
+		* threads call SetListsOnly(true) and SetCacheSize(CACHESIZE_GEFORCE3). If
+		* NvTriStrip is ever used with different settings the library will need
+		* some modifications to be thread-safe. ***
+		*/
+		template<typename IndexDataType, typename Allocator>
+		void CacheOptimizeIndexBuffer(TArray<IndexDataType, Allocator>& Indices)
+		{
+			static_assert(sizeof(IndexDataType) == 2 || sizeof(IndexDataType) == 4, "Indices must be short or int.");
+
+			PrimitiveGroup*	PrimitiveGroups = NULL;
+			uint32			NumPrimitiveGroups = 0;
+			bool Is32Bit = sizeof(IndexDataType) == 4;
+
+			SetListsOnly(true);
+			SetCacheSize(CACHESIZE_GEFORCE3);
+
+			GenerateStrips((uint8*)Indices.GetData(), Is32Bit, Indices.Num(), &PrimitiveGroups, &NumPrimitiveGroups);
+
+			Indices.Empty();
+			Indices.AddUninitialized(PrimitiveGroups->numIndices);
+
+			if (Is32Bit)
+			{
+				FMemory::Memcpy(Indices.GetData(), PrimitiveGroups->indices, Indices.Num() * sizeof(IndexDataType));
+			}
+			else
+			{
+				for (uint32 I = 0; I < PrimitiveGroups->numIndices; ++I)
+				{
+					Indices[I] = (uint16)PrimitiveGroups->indices[I];
+				}
+			}
+
+			delete[] PrimitiveGroups;
+		}
+	}
+
+	/*------------------------------------------------------------------------------
+	Forsyth algorithm for cache optimizing index buffers.
+	------------------------------------------------------------------------------*/
+
+	namespace ForsythHelper
+	{
+		/**
+		* Converts 16 bit indices to 32 bit prior to passing them into the real OptimizeFaces util method
+		*/
+		void OptimizeFaces(
+			const uint8* Indices,
+			bool Is32Bit,
+			const uint32 NumIndices,
+			uint32 NumVertices,
+			uint32* OutIndices,
+			uint16 CacheSize
+		)
+		{
+			if (Is32Bit)
+			{
+				Forsyth::OptimizeFaces((uint32*)Indices, NumIndices, NumVertices, OutIndices, CacheSize);
+			}
+			else
+			{
+				// convert to 32 bit
+				uint32 Idx;
+				TArray<uint32> NewIndices;
+				NewIndices.AddUninitialized(NumIndices);
+				for (Idx = 0; Idx < NumIndices; ++Idx)
+				{
+					NewIndices[Idx] = ((uint16*)Indices)[Idx];
+				}
+				Forsyth::OptimizeFaces(NewIndices.GetData(), NumIndices, NumVertices, OutIndices, CacheSize);
+			}
+
+		}
+
+		/**
+		* Orders a triangle list for better vertex cache coherency.
+		*/
+		template<typename IndexDataType, typename Allocator>
+		void CacheOptimizeIndexBuffer(TArray<IndexDataType, Allocator>& Indices)
+		{
+			static_assert(sizeof(IndexDataType) == 2 || sizeof(IndexDataType) == 4, "Indices must be short or int.");
+			bool Is32Bit = sizeof(IndexDataType) == 4;
+
+			// Count the number of vertices
+			uint32 NumVertices = 0;
+			for (int32 Index = 0; Index < Indices.Num(); ++Index)
+			{
+				if (Indices[Index] > NumVertices)
+				{
+					NumVertices = Indices[Index];
+				}
+			}
+			NumVertices += 1;
+
+			TArray<uint32> OptimizedIndices;
+			OptimizedIndices.AddUninitialized(Indices.Num());
+			uint16 CacheSize = 32;
+			OptimizeFaces((uint8*)Indices.GetData(), Is32Bit, Indices.Num(), NumVertices, OptimizedIndices.GetData(), CacheSize);
+
+			if (Is32Bit)
+			{
+				FMemory::Memcpy(Indices.GetData(), OptimizedIndices.GetData(), Indices.Num() * sizeof(IndexDataType));
+			}
+			else
+			{
+				for (int32 I = 0; I < OptimizedIndices.Num(); ++I)
+				{
+					Indices[I] = (uint16)OptimizedIndices[I];
+				}
+			}
+		}
+	}
+}
+
 class BuildOptimizationHelper
 {
 public:
@@ -385,4 +564,83 @@ public:
 			OutPnAenIndices.Empty();
 		}
 	}
+
+	static void CacheOptimizeIndexBuffer(TArray<uint32>& Indices)
+	{
+		bool bDisableTriangleOrderOptimization = (BuildOptimizationPrivate::CVarTriangleOrderOptimization.GetValueOnGameThread() == 2);
+		bool bUsingNvTriStrip = !bDisableTriangleOrderOptimization && (BuildOptimizationPrivate::CVarTriangleOrderOptimization.GetValueOnGameThread() == 0);
+
+		if (bUsingNvTriStrip)
+		{
+			BuildOptimizationPrivate::NvTriStripHelper::CacheOptimizeIndexBuffer(Indices);
+		}
+		else if (!bDisableTriangleOrderOptimization)
+		{
+			BuildOptimizationPrivate::ForsythHelper::CacheOptimizeIndexBuffer(Indices);
+		}
+	}
+
+	static void CacheOptimizeVertexAndIndexBuffer(
+		TArray<FStaticMeshBuildVertex>& Vertices,
+		TArray<TArray<uint32> >& PerSectionIndices,
+		TArray<int32>& WedgeMap
+	)
+	{
+		// Copy the vertices since we will be reordering them
+		TArray<FStaticMeshBuildVertex> OriginalVertices = Vertices;
+
+		// Initialize a cache that stores which indices have been assigned
+		TArray<int32> IndexCache;
+		IndexCache.AddUninitialized(Vertices.Num());
+		FMemory::Memset(IndexCache.GetData(), INDEX_NONE, IndexCache.Num() * IndexCache.GetTypeSize());
+		int32 NextAvailableIndex = 0;
+
+		// Iterate through the section index buffers, 
+		// Optimizing index order for the post transform cache (minimizes the number of vertices transformed), 
+		// And vertex order for the pre transform cache (minimizes the amount of vertex data fetched by the GPU).
+		for (int32 SectionIndex = 0; SectionIndex < PerSectionIndices.Num(); SectionIndex++)
+		{
+			TArray<uint32>& Indices = PerSectionIndices[SectionIndex];
+
+			if (Indices.Num())
+			{
+				// Optimize the index buffer for the post transform cache with.
+				CacheOptimizeIndexBuffer(Indices);
+
+				// Copy the index buffer since we will be reordering it
+				TArray<uint32> OriginalIndices = Indices;
+
+				// Go through the indices and assign them new values that are coherent where possible
+				for (int32 Index = 0; Index < Indices.Num(); Index++)
+				{
+					const int32 CachedIndex = IndexCache[OriginalIndices[Index]];
+
+					if (CachedIndex == INDEX_NONE)
+					{
+						// No new index has been allocated for this existing index, assign a new one
+						Indices[Index] = NextAvailableIndex;
+						// Mark what this index has been assigned to
+						IndexCache[OriginalIndices[Index]] = NextAvailableIndex;
+						NextAvailableIndex++;
+					}
+					else
+					{
+						// Reuse an existing index assignment
+						Indices[Index] = CachedIndex;
+					}
+					// Reorder the vertices based on the new index assignment
+					Vertices[Indices[Index]] = OriginalVertices[OriginalIndices[Index]];
+				}
+			}
+		}
+
+		for (int32 i = 0; i < WedgeMap.Num(); i++)
+		{
+			int32 MappedIndex = WedgeMap[i];
+			if (MappedIndex != INDEX_NONE)
+			{
+				WedgeMap[i] = IndexCache[MappedIndex];
+			}
+		}
+	}
 };
\ No newline at end of file
diff --git a/Engine/Source/Runtime/MeshBuilder/Private/StaticMeshBuilder.cpp b/Engine/Source/Runtime/MeshBuilder/Private/StaticMeshBuilder.cpp
--- a/Engine/Source/Runtime/MeshBuilder/Private/StaticMeshBuilder.cpp
+++ b/Engine/Source/Runtime/MeshBuilder/Private/StaticMeshBuilder.cpp
@@ -15,7 +15,7 @@ bool IsOrphanedVertex(const class UMeshDescription *MeshDescription, const FVert
 void UpdateBounds(class UStaticMesh* StaticMesh);
 void UpdateCollision(UStaticMesh *StaticMesh);
 void BuildVertexBuffer(const class UMeshDescription* MeshDescription, struct FStaticMeshLODResources& StaticMeshLOD, const struct FMeshBuildSettings& LODBuildSettings, TArray< FStaticMeshBuildVertex >& StaticMeshBuildVertices);
-void BuildIndexBuffer(class UStaticMesh *StaticMesh, const class UMeshDescription* MeshDescription, TArray< uint32 >& IndexBuffer, struct FStaticMeshLODResources& StaticMeshLOD, TArray<int32>& OutWedgeMap);
+void BuildIndexBuffer(class UStaticMesh *StaticMesh, const class UMeshDescription* MeshDescription, TArray< uint32 >& IndexBuffer, struct FStaticMeshLODResources& StaticMeshLOD, TArray<int32>& OutWedgeMap, TArray<TArray<uint32> >& OutPerSectionIndices);
 void BuildAllBufferOptimizations(struct FStaticMeshLODResources& StaticMeshLOD, const struct FMeshBuildSettings& LODBuildSettings, TArray< uint32 >& IndexBuffer, bool bNeeds32BitIndices, TArray< FStaticMeshBuildVertex >& StaticMeshBuildVertices);
 //////////////////////////////////////////////////////////////////////////
 
@@ -32,6 +32,7 @@ bool FStaticMeshBuilder::Build(UStaticMesh* StaticMesh)
 		//TODO: Warn the user that there is no mesh description data
 		return false;
 	}
+	StaticMesh->RenderData->AllocateLODResources(MeshDescriptionCount);
 
 	OnBuildRenderMeshStart(StaticMesh, false);
 
@@ -63,7 +64,20 @@ bool FStaticMeshBuilder::Build(UStaticMesh* StaticMesh)
 		//Render data Wedge map is only set for LOD 0???
 		TArray<int32> TempWedgeMap;
 		TArray<int32> &WedgeMap = (LodIndex == 0) ? StaticMeshRenderData.WedgeMap : TempWedgeMap;
-		BuildIndexBuffer(StaticMesh, MeshDescription, IndexBuffer, StaticMeshLOD, WedgeMap);
+
+		//Prepare the PerSectionIndices array so we can optimize the index buffer for the GPU
+		int32 MaxMaterialIndex = 1;
+		for (const FPolygonGroupID& PolygonGroupID : MeshDescription->PolygonGroups().GetElementIDs())
+		{
+			MaxMaterialIndex = FMath::Max<int32>(PolygonGroupID.GetValue(), MaxMaterialIndex);
+		}
+		TArray<TArray<uint32> > PerSectionIndices;
+		for (int32 i = 0; i <= MaxMaterialIndex; ++i)
+		{
+			PerSectionIndices.Push(TArray<uint32>());
+		}
+
+		BuildIndexBuffer(StaticMesh, MeshDescription, IndexBuffer, StaticMeshLOD, WedgeMap, PerSectionIndices);
 
 		// Figure out which index buffer stride we need
 		bool bNeeds32BitIndices = false;
@@ -76,6 +90,11 @@ bool FStaticMeshBuilder::Build(UStaticMesh* StaticMesh)
 		}
 		const EIndexBufferStride::Type IndexBufferStride = bNeeds32BitIndices ? EIndexBufferStride::Force32Bit : EIndexBufferStride::Force16Bit;
 		StaticMeshLOD.IndexBuffer.SetIndices(IndexBuffer, IndexBufferStride);
+		if (MeshDescription->VertexInstances().Num() < 100000 * 3)
+		{
+			BuildOptimizationHelper::CacheOptimizeVertexAndIndexBuffer(StaticMeshBuildVertices, PerSectionIndices, WedgeMap);
+			check(WedgeMap.Num() == MeshDescription->VertexInstances().Num());
+		}
 
 		BuildAllBufferOptimizations(StaticMeshLOD, LODBuildSettings, IndexBuffer, bNeeds32BitIndices, StaticMeshBuildVertices);
 	} //End of LOD for loop
@@ -305,7 +324,7 @@ void BuildVertexBuffer(const UMeshDescription* MeshDescription, FStaticMeshLODRe
 	}
 }
 
-void BuildIndexBuffer(UStaticMesh *StaticMesh, const UMeshDescription* MeshDescription, TArray< uint32 >& IndexBuffer, FStaticMeshLODResources& StaticMeshLOD, TArray<int32>& OutWedgeMap)
+void BuildIndexBuffer(UStaticMesh *StaticMesh, const UMeshDescription* MeshDescription, TArray< uint32 >& IndexBuffer, FStaticMeshLODResources& StaticMeshLOD, TArray<int32>& OutWedgeMap, TArray<TArray<uint32> >& OutPerSectionIndices)
 {
 	OutWedgeMap.Reset();
 
@@ -313,6 +332,7 @@ void BuildIndexBuffer(UStaticMesh *StaticMesh, const UMeshDescription* MeshDescr
 	// Set up index buffer
 	for (const FPolygonGroupID& PolygonGroupID : PolygonGroups.GetElementIDs())
 	{
+		TArray<uint32>& SectionIndices = OutPerSectionIndices[PolygonGroupID.GetValue()];
 		TArray<FMeshTriangle> Triangles;
 		int32 SectionTriangleCount = GetPolygonGroupTriangles(MeshDescription, Triangles, PolygonGroupID);
 
@@ -347,6 +367,7 @@ void BuildIndexBuffer(UStaticMesh *StaticMesh, const UMeshDescription* MeshDescr
 					const uint32 RenderingVertexIndex = Triangle.GetVertexInstanceID(TriVert).GetValue();
 					IndexBuffer.Add(RenderingVertexIndex);
 					OutWedgeMap.Add(RenderingVertexIndex);
+					SectionIndices.Add(RenderingVertexIndex);
 					MinIndex = FMath::Min(MinIndex, RenderingVertexIndex);
 					MaxIndex = FMath::Max(MaxIndex, RenderingVertexIndex);
 				}
@@ -396,6 +417,12 @@ void BuildAllBufferOptimizations(FStaticMeshLODResources& StaticMeshLOD, const F
 			IndexBuffer,
 			StaticMeshLOD.Sections
 		);
+
+		if (DepthOnlyIndices.Num() < 50000 * 3)
+		{
+			BuildOptimizationHelper::CacheOptimizeIndexBuffer(DepthOnlyIndices);
+		}
+
 		StaticMeshLOD.DepthOnlyIndexBuffer.SetIndices(DepthOnlyIndices, IndexBufferStride);
 	}
 
diff --git a/Engine/Source/Runtime/MeshDescription/MeshDescription.Build.cs b/Engine/Source/Runtime/MeshDescription/MeshDescription.Build.cs
--- a/Engine/Source/Runtime/MeshDescription/MeshDescription.Build.cs
+++ b/Engine/Source/Runtime/MeshDescription/MeshDescription.Build.cs
@@ -14,7 +14,6 @@ namespace UnrealBuildTool.Rules
 				{
 					"Core",
 					"CoreUObject",
-                    "Engine"
 				}
 			);
 		}
diff --git a/Engine/Source/Runtime/MeshDescription/Private/MeshDescription.cpp b/Engine/Source/Runtime/MeshDescription/Private/MeshDescription.cpp
--- a/Engine/Source/Runtime/MeshDescription/Private/MeshDescription.cpp
+++ b/Engine/Source/Runtime/MeshDescription/Private/MeshDescription.cpp
@@ -1,7 +1,8 @@
 // Copyright 1998-2017 Epic Games, Inc. All Rights Reserved.
 
 #include "MeshDescription.h"
-
+#include "Misc/SecureHash.h"
+#include "Serialization/MemoryWriter.h"
 
 UMeshDescription::UMeshDescription()
 {
@@ -52,11 +53,29 @@ void UMeshDescription::Serialize( FArchive& Ar )
 	Ar << PolygonGroupAttributesSet;
 }
 
-
-void UMeshDescription::PostLoad()
+#if WITH_EDITORONLY_DATA
+FString UMeshDescription::GetIdString()
 {
+	TArray<uint8> TempBytes;
+	FMemoryWriter Ar(TempBytes, /*bIsPersistent=*/ true);
+	Serialize(Ar);
+	FSHA1 Sha;
+	TArray<TCHAR> OwnerName = GetPathName().GetCharArray();
+	Sha.Update((uint8*)OwnerName.GetData(), OwnerName.Num() * OwnerName.GetTypeSize());
+	if (TempBytes.Num() > 0)
+	{
+		uint8* Buffer = TempBytes.GetData();
+		Sha.Update(Buffer, TempBytes.Num());
+	}
+	Sha.Final();
+	TempBytes.Empty();
+	// Retrieve the hash and use it to construct a pseudo-GUID.
+	uint32 Hash[5];
+	Sha.GetHash((uint8*)Hash);
+	FGuid Guid = FGuid(Hash[0] ^ Hash[4], Hash[1], Hash[2], Hash[3]);
+	return Guid.ToString(EGuidFormats::Digits);
 }
-
+#endif
 
 void UMeshDescription::Compact( FElementIDRemappings& OutRemappings )
 {
diff --git a/Engine/Source/Runtime/MeshDescription/Public/MeshDescription.h b/Engine/Source/Runtime/MeshDescription/Public/MeshDescription.h
--- a/Engine/Source/Runtime/MeshDescription/Public/MeshDescription.h
+++ b/Engine/Source/Runtime/MeshDescription/Public/MeshDescription.h
@@ -4,9 +4,9 @@
 #pragma warning(disable:4503)
 
 #include "CoreMinimal.h"
+#include "CoreTypes.h"
 #include "MeshTypes.h"
 #include "MeshElementArray.h"
-#include "Materials/MaterialInterface.h"
 #include "MeshDescription.generated.h"
 
 
@@ -306,7 +306,9 @@ struct FMeshPolygonGroup
 	/** Serializer */
 	friend FArchive& operator<<( FArchive& Ar, FMeshPolygonGroup& PolygonGroup )
 	{
-		Ar << PolygonGroup.MaterialAsset;
+		//TODO FStringAssetReference cannot be archive with FArchive it can be archive only with FArchiveUObject
+		//Ar << PolygonGroup.MaterialAsset;
+		PolygonGroup.MaterialAsset.SerializePath(Ar);
 		Ar << PolygonGroup.MaterialSlotName;
 		Ar << PolygonGroup.ImportedMaterialSlotName;
 		Ar << PolygonGroup.bEnableCollision;
@@ -569,7 +571,6 @@ public:
 
 	// UObject interface
 	virtual void Serialize( FArchive& Ar ) override;
-	virtual void PostLoad() override;
 
 	FVertexArray& Vertices() { return VertexArray; }
 	const FVertexArray& Vertices() const { return VertexArray; }
@@ -636,6 +637,10 @@ public:
 		VertexArray.Remove( VertexID );
 	}
 
+#if WITH_EDITORONLY_DATA
+	FString GetIdString();
+#endif
+
 private:
 	inline void CreateVertexInstance_Internal( const FVertexInstanceID VertexInstanceID, const FVertexID VertexID )
 	{
@@ -880,4 +885,4 @@ private:
 	TAttributesSet<FEdgeID> EdgeAttributesSet;
 	TAttributesSet<FPolygonID> PolygonAttributesSet;
 	TAttributesSet<FPolygonGroupID> PolygonGroupAttributesSet;
-};
\ No newline at end of file
+};
