Normals and mikkt space tangents build for MeshDescription

diff --git a/Engine/Source/Editor/UnrealEd/Private/Fbx/FbxStaticMeshImport.cpp b/Engine/Source/Editor/UnrealEd/Private/Fbx/FbxStaticMeshImport.cpp
--- a/Engine/Source/Editor/UnrealEd/Private/Fbx/FbxStaticMeshImport.cpp
+++ b/Engine/Source/Editor/UnrealEd/Private/Fbx/FbxStaticMeshImport.cpp
@@ -265,7 +265,7 @@ bool UnFbx::FFbxImporter::BuildStaticMeshFromGeometry(FbxNode* Node, UStaticMesh
 	FbxMesh* Mesh = Node->GetMesh();
 	FStaticMeshSourceModel& SrcModel = StaticMesh->SourceModels[LODIndex];
 	
-	UMeshDescription *MeshDescription = StaticMesh->GetMeshDescription(LODIndex);
+	UMeshDescription *MeshDescription = StaticMesh->GetOriginalMeshDescription(LODIndex);
 	if (ImportOptions->bImportMeshDescription)
 	{
 		//The mesh description should have been created before calling BuildStaticMeshFromGeometry
@@ -1554,13 +1554,14 @@ UStaticMesh* UnFbx::FFbxImporter::ImportStaticMeshAsSingle(UObject* InParent, TA
 
 	if (ImportOptions->bImportMeshDescription)
 	{
-		MeshDescription = StaticMesh->GetMeshDescription(LODIndex);
+		MeshDescription = StaticMesh->GetOriginalMeshDescription(LODIndex);
 		if (MeshDescription == nullptr)
 		{
+			StaticMesh->ClearOriginalMeshDescription(LODIndex);
 			//Create private asset in the same package as the StaticMesh, and make sure reference are set to avoid GC
 			MeshDescription = NewObject<UMeshDescription>(StaticMesh, NAME_None, RF_NoFlags);
 			check(MeshDescription != nullptr);
-			StaticMesh->MeshDescriptions.Add(MeshDescription);
+			StaticMesh->SetOriginalMeshDescription(LODIndex, MeshDescription);
 		}
 	}
 
@@ -1583,7 +1584,7 @@ UStaticMesh* UnFbx::FFbxImporter::ImportStaticMeshAsSingle(UObject* InParent, TA
 		SrcModel.RawMeshBulkData->SaveRawMesh( EmptyRawMesh );
 	}
 
-	StaticMesh->ClearOriginalMeshDescription(LODIndex);
+	
 	
 	// make sure it has a new lighting guid
 	StaticMesh->LightingGuid = FGuid::NewGuid();
diff --git a/Engine/Source/Runtime/Engine/Classes/Engine/StaticMesh.h b/Engine/Source/Runtime/Engine/Classes/Engine/StaticMesh.h
--- a/Engine/Source/Runtime/Engine/Classes/Engine/StaticMesh.h
+++ b/Engine/Source/Runtime/Engine/Classes/Engine/StaticMesh.h
@@ -627,6 +627,7 @@ public:
 	TArray<class UMeshDescription*> MeshDescriptions;
 
 	ENGINE_API class UMeshDescription* GetMeshDescription(int32 LodIndex=0) const;
+	ENGINE_API void SetMeshDescription(int32 LodIndex, class UMeshDescription* InMeshDescription);
 	ENGINE_API int32 GetMeshDescriptionCount() const;
 
 #if WITH_EDITORONLY_DATA
diff --git a/Engine/Source/Runtime/Engine/Private/StaticMesh.cpp b/Engine/Source/Runtime/Engine/Private/StaticMesh.cpp
--- a/Engine/Source/Runtime/Engine/Private/StaticMesh.cpp
+++ b/Engine/Source/Runtime/Engine/Private/StaticMesh.cpp
@@ -1293,7 +1293,7 @@ void FStaticMeshRenderData::Cache(UStaticMesh* Owner, const FStaticMeshLODSettin
 			Args.Add(TEXT("StaticMeshName"), FText::FromString( Owner->GetName() ) );
 			FStaticMeshStatusMessageContext StatusContext( FText::Format( NSLOCTEXT("Engine", "BuildingStaticMeshStatus", "Building static mesh {StaticMeshName}..."), Args ) );
 
-			if (Owner->GetMeshDescriptionCount() > 0)
+			if (Owner->GetOriginalMeshDescriptionCount() > 0)
 			{
 				IMeshBuilderModule& MeshBuilderModule = FModuleManager::Get().LoadModuleChecked<IMeshBuilderModule>(TEXT("MeshBuilder"));
 				MeshBuilderModule.BuildMesh(Owner);
@@ -2213,6 +2213,17 @@ UMeshDescription* UStaticMesh::GetMeshDescription(int32 LodIndex) const
 	return MeshDescriptions[LodIndex];
 }
 
+void UStaticMesh::SetMeshDescription(int32 LodIndex, class UMeshDescription* InMeshDescription)
+{
+	if (!MeshDescriptions.IsValidIndex(LodIndex))
+	{
+		//Add nullptr missing entries
+		MeshDescriptions.AddZeroed(LodIndex - MeshDescriptions.Num() + 1);
+	}
+	//Set the original mesh description
+	MeshDescriptions[LodIndex] = InMeshDescription;
+}
+
 int32 UStaticMesh::GetMeshDescriptionCount() const
 {
 	return MeshDescriptions.Num();
diff --git a/Engine/Source/Runtime/MeshBuilder/MeshBuilder.Build.cs b/Engine/Source/Runtime/MeshBuilder/MeshBuilder.Build.cs
--- a/Engine/Source/Runtime/MeshBuilder/MeshBuilder.Build.cs
+++ b/Engine/Source/Runtime/MeshBuilder/MeshBuilder.Build.cs
@@ -17,12 +17,30 @@ namespace UnrealBuildTool.Rules
                     "Engine",
                     "RenderCore",
                     "MeshDescription",
+                    "MeshReductionInterface",
+                    "RenderCore"
                 }
 			);
 
             AddEngineThirdPartyPrivateStaticDependencies(Target, "nvTriStrip");
             AddEngineThirdPartyPrivateStaticDependencies(Target, "ForsythTriOptimizer");
             AddEngineThirdPartyPrivateStaticDependencies(Target, "nvTessLib");
+            AddEngineThirdPartyPrivateStaticDependencies(Target, "MikkTSpace");
+            AddEngineThirdPartyPrivateStaticDependencies(Target, "QuadricMeshReduction");
+
+            if (Target.bCompileSimplygon == true)
+            {
+                AddEngineThirdPartyPrivateDynamicDependencies(Target, "SimplygonMeshReduction");
+
+                if (Target.bCompileSimplygonSSF == true)
+                {
+                    DynamicallyLoadedModuleNames.AddRange(
+                        new string[] {
+                    "SimplygonSwarm"
+                    }
+                    );
+                }
+            }
         }
 	}
 }
diff --git a/Engine/Source/Runtime/MeshBuilder/Private/MeshDescriptionHelper.cpp b/Engine/Source/Runtime/MeshBuilder/Private/MeshDescriptionHelper.cpp
new file mode 100644
--- /dev/null
+++ b/Engine/Source/Runtime/MeshBuilder/Private/MeshDescriptionHelper.cpp
@@ -0,0 +1,663 @@
+// Copyright 1998-2017 Epic Games, Inc. All Rights Reserved.
+#pragma once
+
+#include "MeshDescriptionHelper.h"
+
+#include "CoreMinimal.h"
+#include "CoreTypes.h"
+#include "Engine/EngineTypes.h"
+#include "MeshDescription.h"
+#include "mikktspace.h"
+#include "Serialization/MemoryWriter.h"
+#include "Serialization/MemoryReader.h"
+#include "RenderUtils.h"
+
+namespace MeshDescriptionMikktSpaceInterface
+{
+	//Mikk t spce static function
+	int MikkGetNumFaces(const SMikkTSpaceContext* Context);
+	int MikkGetNumVertsOfFace(const SMikkTSpaceContext* Context, const int FaceIdx);
+	void MikkGetPosition(const SMikkTSpaceContext* Context, float Position[3], const int FaceIdx, const int VertIdx);
+	void MikkGetNormal(const SMikkTSpaceContext* Context, float Normal[3], const int FaceIdx, const int VertIdx);
+	void MikkSetTSpaceBasic(const SMikkTSpaceContext* Context, const float Tangent[3], const float BitangentSign, const int FaceIdx, const int VertIdx);
+	void MikkGetTexCoord(const SMikkTSpaceContext* Context, float UV[2], const int FaceIdx, const int VertIdx);
+}
+
+FMeshDescriptionHelper::FMeshDescriptionHelper(FMeshBuildSettings* InBuildSettings, const UMeshDescription* InOriginalMeshDescription)
+	: OriginalMeshDescription(InOriginalMeshDescription)
+	, BuildSettings(InBuildSettings)
+{
+}
+
+UMeshDescription* FMeshDescriptionHelper::GetRenderMeshDescription(UObject* Owner)
+{
+	//Use the build settings to create the RenderMeshDescription
+	UMeshDescription *RenderMeshDescription = NewObject<UMeshDescription>(Owner, NAME_None, RF_NoFlags);
+
+	if(OriginalMeshDescription != nullptr)
+	{
+		//Copy The Original Mesh Description in the render mesh description
+		CopyMeshDescription(const_cast<UMeshDescription*>(OriginalMeshDescription), RenderMeshDescription);
+		float ComparisonThreshold = BuildSettings->bRemoveDegenerates ? THRESH_POINTS_ARE_SAME : 0.0f;
+		int32 NumWedges = RenderMeshDescription->VertexInstances().Num();
+		// Find overlapping corners to accelerate adjacency.
+		FindOverlappingCorners(ComparisonThreshold);
+
+		// Figure out if we should recompute normals and tangents.
+		int32 OriginalNormalCount = RenderMeshDescription->VertexInstances().Num();
+		int32 OriginalTangentCount = OriginalNormalCount;
+		int32 OriginalBiNormalCount = OriginalNormalCount;
+		bool bRecomputeNormals = BuildSettings->bRecomputeNormals || OriginalNormalCount != NumWedges;
+		bool bRecomputeTangents = BuildSettings->bRecomputeTangents || OriginalTangentCount != NumWedges || OriginalBiNormalCount != NumWedges;
+
+		if (bRecomputeTangents || bRecomputeNormals)
+		{
+			for (const FVertexInstanceID& VertexInstanceID : RenderMeshDescription->VertexInstances().GetElementIDs())
+			{
+				FMeshVertexInstance& VertexInstance = RenderMeshDescription->GetVertexInstance(VertexInstanceID);
+				// Dump normals and tangents if we are recomputing them.
+				if (bRecomputeTangents)
+				{
+					//Dump the tangents
+					VertexInstance.BinormalSign = 0.0f;
+					VertexInstance.Tangent = FVector(0.0f);
+				}
+				if (bRecomputeNormals)
+				{
+					//Dump the normals
+					VertexInstance.Normal = FVector(0.0f);
+				}
+			}
+		}
+
+		// Compute any missing tangents.
+		{
+			// Static meshes always blend normals of overlapping corners.
+			uint32 TangentOptions = FMeshDescriptionHelper::ETangentOptions::BlendOverlappingNormals;
+			if (BuildSettings->bRemoveDegenerates)
+			{
+				// If removing degenerate triangles, ignore them when computing tangents.
+				TangentOptions |= FMeshDescriptionHelper::ETangentOptions::IgnoreDegenerateTriangles;
+			}
+
+			//Keep the original mesh description NTBs if we do not rebuild the normals or tangents.
+			//MikkTSpace should be use only when the user want to recompute the normals or tangents otherwise should always fallback on builtin
+			if (BuildSettings->bUseMikkTSpace && (BuildSettings->bRecomputeNormals || BuildSettings->bRecomputeTangents))
+			{
+				ComputeNTB_MikkTSpace(RenderMeshDescription, (FMeshDescriptionHelper::ETangentOptions)TangentOptions);
+			}
+			//TODO: builtin Compute of the tangents
+			//else if(BuildSettings->bRecomputeNormals || BuildSettings->bRecomputeTangents)
+			//{
+				//ComputeTangents(RenderMeshDescription, OverlappingCorners, TangentOptions);
+			//}
+		}
+
+		// TODO: Generate lightmap UVs
+		/*if (BuildSettings->bGenerateLightmapUVs)
+		{
+			if (RawMesh.WedgeTexCoords[BuildSettings->SrcLightmapIndex].Num() == 0)
+			{
+				BuildSettings->SrcLightmapIndex = 0;
+			}
+
+			FLayoutUV Packer(&RawMesh, BuildSettings->SrcLightmapIndex, BuildSettings->DstLightmapIndex, BuildSettings->MinLightmapResolution);
+			Packer.SetVersion(LightmapUVVersion);
+
+			Packer.FindCharts(OverlappingCorners);
+			bool bPackSuccess = Packer.FindBestPacking();
+			if (bPackSuccess)
+			{
+				Packer.CommitPackedUVs();
+			}
+		}
+		HasRawMesh[LODIndex] = true;*/
+	}
+	else
+	{
+	}
+
+	return RenderMeshDescription;
+}
+
+void FMeshDescriptionHelper::CopyMeshDescription(UMeshDescription* SourceMeshDescription, UMeshDescription* DestinationMeshDescription) const
+{
+	//Copy the Source into the destination
+	//Save the Source
+	TArray<uint8> TempBytes;
+	FMemoryWriter SaveAr(TempBytes, /*bIsPersistent=*/ true);
+	SourceMeshDescription->Serialize(SaveAr);
+	//Load the save array in the destination
+	FMemoryReader LoadAr(TempBytes, /*bIsPersistent=*/ true);
+	DestinationMeshDescription->Serialize(LoadAr);
+}
+
+bool FMeshDescriptionHelper::IsValidOriginalMeshDescription()
+{
+	return OriginalMeshDescription != nullptr;
+}
+
+void FMeshDescriptionHelper::FindOverlappingCorners(float ComparisonThreshold)
+{
+	//Empty the old data
+	OverlappingCorners.Empty();
+	const int32 NumWedges = OriginalMeshDescription->VertexInstances().Num();
+
+	// Create a list of vertex Z/index pairs
+	TArray<FIndexAndZ> VertIndexAndZ;
+	VertIndexAndZ.Reserve(NumWedges);
+	const FVertexArray& Vertices = OriginalMeshDescription->Vertices();
+	const FVertexInstanceArray& VertexInstances = OriginalMeshDescription->VertexInstances();
+
+	for (const FPolygonID& PolygonID : OriginalMeshDescription->Polygons().GetElementIDs())
+	{
+		const TArray<FMeshTriangle>& MeshTriangles = OriginalMeshDescription->GetPolygon(PolygonID).Triangles;
+		for (const FMeshTriangle& MeshTriangle : MeshTriangles)
+		{
+			for (int32 CornerIndex = 0; CornerIndex < 3; ++CornerIndex)
+			{
+				const FVertexInstanceID& VertexInstanceID = MeshTriangle.GetVertexInstanceID(CornerIndex);
+				const FMeshVertexInstance& VertexInstance = VertexInstances[VertexInstanceID];
+
+				new(VertIndexAndZ)FIndexAndZ(VertexInstanceID.GetValue(), Vertices[VertexInstance.VertexID].VertexPosition);
+			}
+		}
+	}
+
+	// Sort the vertices by z value
+	VertIndexAndZ.Sort(FCompareIndexAndZ());
+
+	// Search for duplicates, quickly!
+	for (int32 i = 0; i < VertIndexAndZ.Num(); i++)
+	{
+		// only need to search forward, since we add pairs both ways
+		for (int32 j = i + 1; j < VertIndexAndZ.Num(); j++)
+		{
+			if (FMath::Abs(VertIndexAndZ[j].Z - VertIndexAndZ[i].Z) > ComparisonThreshold)
+				break; // can't be any more dups
+
+			const FVector& PositionA = OriginalMeshDescription->GetVertex(OriginalMeshDescription->GetVertexInstance(FVertexInstanceID(VertIndexAndZ[i].Index)).VertexID).VertexPosition;
+			const FVector& PositionB = OriginalMeshDescription->GetVertex(OriginalMeshDescription->GetVertexInstance(FVertexInstanceID(VertIndexAndZ[j].Index)).VertexID).VertexPosition;
+
+			if (PositionA.Equals(PositionB, ComparisonThreshold))
+			{
+				OverlappingCorners.Add(VertIndexAndZ[i].Index, VertIndexAndZ[j].Index);
+				OverlappingCorners.Add(VertIndexAndZ[j].Index, VertIndexAndZ[i].Index);
+			}
+		}
+	}
+}
+
+const FVector& FMeshDescriptionHelper::GetVertexPositionFromVertexInstance(UMeshDescription* MeshDescription, int32 VertexInstanceIndex) const
+{
+	return MeshDescription->GetVertex(MeshDescription->GetVertexInstance(FVertexInstanceID(VertexInstanceIndex)).VertexID).VertexPosition;
+}
+
+FVector2D FMeshDescriptionHelper::GetVertexInstanceUV(UMeshDescription* MeshDescription, int32 VertexInstanceIndex, int32 UVLayer) const
+{
+	return MeshDescription->VertexInstanceAttributes().GetAttribute<FVector2D>(FVertexInstanceID(VertexInstanceIndex), UEditableMeshAttribute::VertexTextureCoordinate(), UVLayer);
+}
+
+void FMeshDescriptionHelper::GetVertexInstanceNTB(UMeshDescription* MeshDescription, const FVertexInstanceID& VertexInstanceID, FVector &OutNormal, FVector &OutTangent, FVector &OutBiNormal) const
+{
+	const FMeshVertexInstance& VertexInstance = MeshDescription->GetVertexInstance(VertexInstanceID);
+	OutNormal = VertexInstance.Normal;
+	OutTangent = VertexInstance.Tangent;
+	if (!OutNormal.IsNearlyZero() && !OutTangent.IsNearlyZero())
+	{
+		OutBiNormal = FVector::CrossProduct(VertexInstance.Normal, VertexInstance.Tangent).GetSafeNormal() * VertexInstance.BinormalSign;
+	}
+	else
+	{
+		OutBiNormal = FVector(0.0f);
+	}
+}
+
+void FMeshDescriptionHelper::SetVertexInstanceNTB(UMeshDescription* MeshDescription, const FVertexInstanceID& VertexInstanceID, FVector &OutNormal, FVector &OutTangent, FVector &OutBiNormal)
+{
+	FMeshVertexInstance& VertexInstance = MeshDescription->GetVertexInstance(VertexInstanceID);
+	VertexInstance.Normal = OutNormal;
+	VertexInstance.Tangent = OutTangent;
+	VertexInstance.BinormalSign = GetBasisDeterminantSign(OutTangent, OutBiNormal, OutNormal);
+}
+
+void FMeshDescriptionHelper::ComputeTriangleTangents(
+	UMeshDescription* MeshDescription,
+	TArray<FVector>& OutTangentX,
+	TArray<FVector>& OutTangentY,
+	TArray<FVector>& OutTangentZ,
+	float ComparisonThreshold
+)
+{
+	const int32 NumTriangles = MeshDescription->Polygons().Num();
+	OutTangentX.Empty(NumTriangles);
+	OutTangentY.Empty(NumTriangles);
+	OutTangentZ.Empty(NumTriangles);
+
+	for (const FPolygonID& PolygonID : MeshDescription->Polygons().GetElementIDs())
+	{
+		const TArray<FMeshTriangle>& MeshTriangles = MeshDescription->GetPolygon(PolygonID).Triangles;
+		for (const FMeshTriangle& MeshTriangle : MeshTriangles)
+		{
+			int32 UVIndex = 0;
+
+			FVector P[3];
+			for (int32 i = 0; i < 3; ++i)
+			{
+				P[i] = GetVertexPositionFromVertexInstance(MeshDescription, MeshTriangle.GetVertexInstanceID(i).GetValue());
+			}
+
+			const FVector Normal = ((P[1] - P[2]) ^ (P[0] - P[2])).GetSafeNormal(ComparisonThreshold);
+			FMatrix	ParameterToLocal(
+				FPlane(P[1].X - P[0].X, P[1].Y - P[0].Y, P[1].Z - P[0].Z, 0),
+				FPlane(P[2].X - P[0].X, P[2].Y - P[0].Y, P[2].Z - P[0].Z, 0),
+				FPlane(P[0].X, P[0].Y, P[0].Z, 0),
+				FPlane(0, 0, 0, 1)
+			);
+
+			const FVector2D T1 = GetVertexInstanceUV(MeshDescription, MeshTriangle.GetVertexInstanceID(0).GetValue(), 0);
+			const FVector2D T2 = GetVertexInstanceUV(MeshDescription, MeshTriangle.GetVertexInstanceID(1).GetValue(), 0);
+			const FVector2D T3 = GetVertexInstanceUV(MeshDescription, MeshTriangle.GetVertexInstanceID(2).GetValue(), 0);
+
+			FMatrix ParameterToTexture(
+				FPlane(T2.X - T1.X, T2.Y - T1.Y, 0, 0),
+				FPlane(T3.X - T1.X, T3.Y - T1.Y, 0, 0),
+				FPlane(T1.X, T1.Y, 1, 0),
+				FPlane(0, 0, 0, 1)
+			);
+
+			// Use InverseSlow to catch singular matrices.  Inverse can miss this sometimes.
+			const FMatrix TextureToLocal = ParameterToTexture.Inverse() * ParameterToLocal;
+
+			OutTangentX.Add(TextureToLocal.TransformVector(FVector(1, 0, 0)).GetSafeNormal());
+			OutTangentY.Add(TextureToLocal.TransformVector(FVector(0, 1, 0)).GetSafeNormal());
+			OutTangentZ.Add(Normal);
+
+			FVector::CreateOrthonormalBasis(
+				OutTangentX[PolygonID.GetValue()],
+				OutTangentY[PolygonID.GetValue()],
+				OutTangentZ[PolygonID.GetValue()]
+			);
+		}
+	}
+
+	check(OutTangentX.Num() == NumTriangles);
+	check(OutTangentY.Num() == NumTriangles);
+	check(OutTangentZ.Num() == NumTriangles);
+}
+
+void FMeshDescriptionHelper::ComputeNTB_MikkTSpace(UMeshDescription* MeshDescription, FMeshDescriptionHelper::ETangentOptions TangentOptions)
+{
+	bool bBlendOverlappingNormals = (TangentOptions & FMeshDescriptionHelper::ETangentOptions::BlendOverlappingNormals) != 0;
+	bool bIgnoreDegenerateTriangles = (TangentOptions & FMeshDescriptionHelper::ETangentOptions::IgnoreDegenerateTriangles) != 0;
+	float ComparisonThreshold = bIgnoreDegenerateTriangles ? THRESH_POINTS_ARE_SAME : 0.0f;
+
+	// Compute per-triangle tangents.
+	TArray<FVector> TriangleTangentX;
+	TArray<FVector> TriangleTangentY;
+	TArray<FVector> TriangleTangentZ;
+
+	ComputeTriangleTangents(
+		MeshDescription,
+		TriangleTangentX,
+		TriangleTangentY,
+		TriangleTangentZ,
+		bIgnoreDegenerateTriangles ? SMALL_NUMBER : 0.0f
+	);
+	TArray<uint32> SmoothingGroupIndices;
+	SmoothingGroupIndices.Reserve(MeshDescription->Polygons().Num());
+	
+	//TODO: support smooth group per edges, we currently do not import smooth group
+	for (const FPolygonID& PolygonID : OriginalMeshDescription->Polygons().GetElementIDs())
+	{
+		SmoothingGroupIndices.Add(1);
+	}
+
+	TArray<FVector> OutTangentX;
+	OutTangentX.AddZeroed(MeshDescription->VertexInstances().Num());
+	TArray<FVector> OutTangentY;
+	OutTangentY.AddZeroed(MeshDescription->VertexInstances().Num());
+	TArray<FVector> OutTangentZ;
+	OutTangentZ.AddZeroed(MeshDescription->VertexInstances().Num());
+
+	for (const FPolygonID& PolygonID : MeshDescription->Polygons().GetElementIDs())
+	{
+		const TArray<FMeshTriangle>& MeshTriangles = MeshDescription->GetPolygon(PolygonID).Triangles;
+		for (const FMeshTriangle& MeshTriangle : MeshTriangles)
+		{
+			for (int32 CornerIndex = 0; CornerIndex < 3; ++CornerIndex)
+			{
+				const FVertexInstanceID& VertexInstanceID = MeshTriangle.GetVertexInstanceID(CornerIndex);
+				GetVertexInstanceNTB(MeshDescription, VertexInstanceID, OutTangentZ[VertexInstanceID.GetValue()], OutTangentX[VertexInstanceID.GetValue()], OutTangentY[VertexInstanceID.GetValue()]);
+			}
+		}
+	}
+	// Declare these out here to avoid reallocations.
+	TArray<FFanFace> RelevantFacesForCorner[3];
+	TArray<int32> AdjacentFaces;
+	TArray<int32> DupVerts;
+
+	int32 NumWedges = MeshDescription->VertexInstances().Num();
+	int32 NumFaces = NumWedges / 3;
+	check(MeshDescription->Polygons().Num() == NumFaces);
+
+	bool bWedgeTSpace = false;
+
+	if (OutTangentX.Num() > 0 && OutTangentY.Num() > 0)
+	{
+		bWedgeTSpace = true;
+		for (int32 WedgeIdx = 0; bWedgeTSpace && WedgeIdx < OutTangentX.Num() && WedgeIdx < OutTangentY.Num(); ++WedgeIdx)
+		{
+			bWedgeTSpace &= (!OutTangentX[WedgeIdx].IsNearlyZero()) && (!OutTangentY[WedgeIdx].IsNearlyZero());
+		}
+	}
+
+	// Allocate storage for tangents if none were provided, and calculate normals for MikkTSpace.
+	if (OutTangentZ.Num() != NumWedges)
+	{
+		// normals are not included, so we should calculate them
+		OutTangentZ.Empty(NumWedges);
+		OutTangentZ.AddZeroed(NumWedges);
+	}
+
+	// we need to calculate normals for MikkTSpace
+	for (const FPolygonID& PolygonID : MeshDescription->Polygons().GetElementIDs())
+	{
+		const TArray<FMeshTriangle>& MeshTriangles = MeshDescription->GetPolygon(PolygonID).Triangles;
+		for (const FMeshTriangle& MeshTriangle : MeshTriangles)
+		{
+			FVector CornerPositions[3];
+			FVector CornerNormal[3];
+
+			for (int32 CornerIndex = 0; CornerIndex < 3; CornerIndex++)
+			{
+				const FVertexInstanceID& VertexInstanceID = MeshTriangle.GetVertexInstanceID(CornerIndex);
+				CornerNormal[CornerIndex] = FVector::ZeroVector;
+				CornerPositions[CornerIndex] = GetVertexPositionFromVertexInstance(MeshDescription, VertexInstanceID.GetValue());
+				RelevantFacesForCorner[CornerIndex].Reset();
+			}
+
+			// Don't process degenerate triangles.
+			if (CornerPositions[0].Equals(CornerPositions[1], ComparisonThreshold)
+				|| CornerPositions[0].Equals(CornerPositions[2], ComparisonThreshold)
+				|| CornerPositions[1].Equals(CornerPositions[2], ComparisonThreshold))
+			{
+				continue;
+			}
+
+			// No need to process triangles if tangents already exist.
+			bool bCornerHasNormal[3] = { 0 };
+			for (int32 CornerIndex = 0; CornerIndex < 3; CornerIndex++)
+			{
+				const FVertexInstanceID& VertexInstanceID = MeshTriangle.GetVertexInstanceID(CornerIndex);
+				bCornerHasNormal[CornerIndex] = !OutTangentZ[VertexInstanceID.GetValue()].IsNearlyZero();
+			}
+			if (bCornerHasNormal[0] && bCornerHasNormal[1] && bCornerHasNormal[2])
+			{
+				continue;
+			}
+
+			// Start building a list of faces adjacent to this face.
+			AdjacentFaces.Reset();
+			for (int32 CornerIndex = 0; CornerIndex < 3; CornerIndex++)
+			{
+				const FVertexInstanceID& VertexInstanceID = MeshTriangle.GetVertexInstanceID(CornerIndex);
+				int32 ThisCornerIndex = VertexInstanceID.GetValue();
+				DupVerts.Reset();
+				OverlappingCorners.MultiFind(ThisCornerIndex, DupVerts);
+				DupVerts.Add(ThisCornerIndex); // I am a "dup" of myself
+				for (int32 k = 0; k < DupVerts.Num(); k++)
+				{
+					AdjacentFaces.AddUnique(DupVerts[k] / 3);
+				}
+			}
+
+			// We need to sort these here because the criteria for point equality is
+			// exact, so we must ensure the exact same order for all dups.
+			AdjacentFaces.Sort();
+
+			// Process adjacent faces
+			for (int32 AdjacentFaceIndex = 0; AdjacentFaceIndex < AdjacentFaces.Num(); AdjacentFaceIndex++)
+			{
+				int32 OtherFaceIndex = AdjacentFaces[AdjacentFaceIndex];
+				for (int32 OurCornerIndex = 0; OurCornerIndex < 3; OurCornerIndex++)
+				{
+					if (bCornerHasNormal[OurCornerIndex])
+						continue;
+
+					FFanFace NewFanFace;
+					int32 CommonIndexCount = 0;
+
+					// Check for vertices in common.
+					if (PolygonID.GetValue() == OtherFaceIndex)
+					{
+						CommonIndexCount = 3;
+						NewFanFace.LinkedVertexIndex = OurCornerIndex;
+					}
+					else
+					{
+						// Check matching vertices against main vertex .
+						for (int32 OtherCornerIndex = 0; OtherCornerIndex < 3; OtherCornerIndex++)
+						{
+							if (CornerPositions[OurCornerIndex].Equals(GetVertexPositionFromVertexInstance(MeshDescription, OtherFaceIndex * 3 + OtherCornerIndex), ComparisonThreshold))
+							{
+								CommonIndexCount++;
+								NewFanFace.LinkedVertexIndex = OtherCornerIndex;
+							}
+						}
+					}
+
+					// Add if connected by at least one point. Smoothing matches are considered later.
+					if (CommonIndexCount > 0)
+					{
+						NewFanFace.FaceIndex = OtherFaceIndex;
+						NewFanFace.bFilled = (OtherFaceIndex == PolygonID.GetValue()); // Starter face for smoothing floodfill.
+						NewFanFace.bBlendTangents = NewFanFace.bFilled;
+						NewFanFace.bBlendNormals = NewFanFace.bFilled;
+						RelevantFacesForCorner[OurCornerIndex].Add(NewFanFace);
+					}
+				}
+			}
+
+			// Find true relevance of faces for a vertex normal by traversing
+			// smoothing-group-compatible connected triangle fans around common vertices.
+			for (int32 CornerIndex = 0; CornerIndex < 3; CornerIndex++)
+			{
+				if (bCornerHasNormal[CornerIndex])
+					continue;
+
+				int32 NewConnections;
+				do
+				{
+					NewConnections = 0;
+					for (int32 OtherFaceIdx = 0; OtherFaceIdx < RelevantFacesForCorner[CornerIndex].Num(); OtherFaceIdx++)
+					{
+						FFanFace& OtherFace = RelevantFacesForCorner[CornerIndex][OtherFaceIdx];
+						// The vertex' own face is initially the only face with bFilled == true.
+						if (OtherFace.bFilled)
+						{
+							for (int32 NextFaceIndex = 0; NextFaceIndex < RelevantFacesForCorner[CornerIndex].Num(); NextFaceIndex++)
+							{
+								FFanFace& NextFace = RelevantFacesForCorner[CornerIndex][NextFaceIndex];
+								if (!NextFace.bFilled) // && !NextFace.bBlendTangents)
+								{
+									if ((NextFaceIndex != OtherFaceIdx)
+										&& (SmoothingGroupIndices[NextFace.FaceIndex] & SmoothingGroupIndices[OtherFace.FaceIndex]))
+									{
+										int32 CommonVertices = 0;
+										int32 CommonNormalVertices = 0;
+										for (int32 OtherCornerIndex = 0; OtherCornerIndex < 3; OtherCornerIndex++)
+										{
+											for (int32 NextCornerIndex = 0; NextCornerIndex < 3; NextCornerIndex++)
+											{
+												int32 NextVertexIndex = MeshDescription->GetPolygon(FPolygonID(NextFace.FaceIndex)).Triangles[0].GetVertexInstanceID(NextCornerIndex).GetValue();
+												int32 OtherVertexIndex = MeshDescription->GetPolygon(FPolygonID(OtherFace.FaceIndex)).Triangles[0].GetVertexInstanceID(OtherCornerIndex).GetValue();
+												if (GetVertexPositionFromVertexInstance(MeshDescription, NextVertexIndex).Equals(GetVertexPositionFromVertexInstance(MeshDescription, OtherVertexIndex), ComparisonThreshold))
+												{
+													CommonVertices++;
+													if (bBlendOverlappingNormals
+														|| NextVertexIndex == OtherVertexIndex)
+													{
+														CommonNormalVertices++;
+													}
+												}
+											}
+										}
+										// Flood fill faces with more than one common vertices which must be touching edges.
+										if (CommonVertices > 1)
+										{
+											NextFace.bFilled = true;
+											NextFace.bBlendNormals = (CommonNormalVertices > 1);
+											NewConnections++;
+										}
+									}
+								}
+							}
+						}
+					}
+				} while (NewConnections > 0);
+			}
+
+			// Vertex normal construction.
+			for (int32 CornerIndex = 0; CornerIndex < 3; CornerIndex++)
+			{
+				const FVertexInstanceID& VertexInstanceID = MeshTriangle.GetVertexInstanceID(CornerIndex);
+				if (bCornerHasNormal[CornerIndex])
+				{
+					CornerNormal[CornerIndex] = OutTangentZ[VertexInstanceID.GetValue()];
+				}
+				else
+				{
+					for (int32 RelevantFaceIdx = 0; RelevantFaceIdx < RelevantFacesForCorner[CornerIndex].Num(); RelevantFaceIdx++)
+					{
+						FFanFace const& RelevantFace = RelevantFacesForCorner[CornerIndex][RelevantFaceIdx];
+						if (RelevantFace.bFilled)
+						{
+							int32 OtherFaceIndex = RelevantFace.FaceIndex;
+							if (RelevantFace.bBlendNormals)
+							{
+								CornerNormal[CornerIndex] += TriangleTangentZ[OtherFaceIndex];
+							}
+						}
+					}
+					if (!OutTangentZ[VertexInstanceID.GetValue()].IsZero())
+					{
+						CornerNormal[CornerIndex] = OutTangentZ[VertexInstanceID.GetValue()];
+					}
+				}
+			}
+
+			// Normalization.
+			for (int32 CornerIndex = 0; CornerIndex < 3; CornerIndex++)
+			{
+				CornerNormal[CornerIndex].Normalize();
+			}
+
+			// Copy back to the outTangentZ (Normal).
+			for (int32 CornerIndex = 0; CornerIndex < 3; CornerIndex++)
+			{
+				const FVertexInstanceID& VertexInstanceID = MeshTriangle.GetVertexInstanceID(CornerIndex);
+				OutTangentZ[VertexInstanceID.GetValue()] = CornerNormal[CornerIndex];
+			}
+		}
+	}
+
+	if (OutTangentX.Num() != NumWedges)
+	{
+		OutTangentX.Empty(NumWedges);
+		OutTangentX.AddZeroed(NumWedges);
+	}
+	if (OutTangentY.Num() != NumWedges)
+	{
+		OutTangentY.Empty(NumWedges);
+		OutTangentY.AddZeroed(NumWedges);
+	}
+
+	//Copy Back the OutTangent to the mesh description
+	for (const FVertexInstanceID& VertexInstanceID : MeshDescription->VertexInstances().GetElementIDs())
+	{
+		int32 VertexInstanceIDValue = VertexInstanceID.GetValue();
+		SetVertexInstanceNTB(MeshDescription, VertexInstanceID, OutTangentZ[VertexInstanceIDValue], OutTangentX[VertexInstanceIDValue], OutTangentY[VertexInstanceIDValue]);
+	}
+
+	if (!bWedgeTSpace)
+	{
+		// we can use mikktspace to calculate the tangents
+		SMikkTSpaceInterface MikkTInterface;
+		MikkTInterface.m_getNormal = MeshDescriptionMikktSpaceInterface::MikkGetNormal;
+		MikkTInterface.m_getNumFaces = MeshDescriptionMikktSpaceInterface::MikkGetNumFaces;
+		MikkTInterface.m_getNumVerticesOfFace = MeshDescriptionMikktSpaceInterface::MikkGetNumVertsOfFace;
+		MikkTInterface.m_getPosition = MeshDescriptionMikktSpaceInterface::MikkGetPosition;
+		MikkTInterface.m_getTexCoord = MeshDescriptionMikktSpaceInterface::MikkGetTexCoord;
+		MikkTInterface.m_setTSpaceBasic = MeshDescriptionMikktSpaceInterface::MikkSetTSpaceBasic;
+		MikkTInterface.m_setTSpace = nullptr;
+
+		SMikkTSpaceContext MikkTContext;
+		MikkTContext.m_pInterface = &MikkTInterface;
+		MikkTContext.m_pUserData = (void*)(MeshDescription);
+		MikkTContext.m_bIgnoreDegenerates = bIgnoreDegenerateTriangles;
+		genTangSpaceDefault(&MikkTContext);
+	}
+
+	check(OutTangentX.Num() == NumWedges);
+	check(OutTangentY.Num() == NumWedges);
+	check(OutTangentZ.Num() == NumWedges);
+}
+
+namespace MeshDescriptionMikktSpaceInterface
+{
+	int MikkGetNumFaces(const SMikkTSpaceContext* Context)
+	{
+		UMeshDescription *MeshDescription = (UMeshDescription*)(Context->m_pUserData);
+		return MeshDescription->Polygons().Num();
+	}
+
+	int MikkGetNumVertsOfFace(const SMikkTSpaceContext* Context, const int FaceIdx)
+	{
+		// All of our meshes are triangles.
+		return 3;
+	}
+
+	void MikkGetPosition(const SMikkTSpaceContext* Context, float Position[3], const int FaceIdx, const int VertIdx)
+	{
+		UMeshDescription *MeshDescription = (UMeshDescription*)(Context->m_pUserData);
+		FVector VertexPosition = MeshDescription->GetVertex(MeshDescription->GetVertexInstance(FVertexInstanceID(FaceIdx * 3 + VertIdx)).VertexID).VertexPosition;
+		Position[0] = VertexPosition.X;
+		Position[1] = VertexPosition.Y;
+		Position[2] = VertexPosition.Z;
+	}
+
+	void MikkGetNormal(const SMikkTSpaceContext* Context, float Normal[3], const int FaceIdx, const int VertIdx)
+	{
+		UMeshDescription *MeshDescription = (UMeshDescription*)(Context->m_pUserData);
+		FVertexInstanceID VertexInstanceID(FaceIdx * 3 + VertIdx);
+		const FMeshVertexInstance& VertexInstance = MeshDescription->GetVertexInstance(VertexInstanceID);
+		const FVector &VertexNormal = VertexInstance.Normal;
+		for (int32 i = 0; i < 3; ++i)
+		{
+			Normal[i] = VertexNormal[i];
+		}
+	}
+
+	void MikkSetTSpaceBasic(const SMikkTSpaceContext* Context, const float Tangent[3], const float BitangentSign, const int FaceIdx, const int VertIdx)
+	{
+		UMeshDescription *MeshDescription = (UMeshDescription*)(Context->m_pUserData);
+		FVertexInstanceID VertexInstanceID(FaceIdx * 3 + VertIdx);
+		FMeshVertexInstance& VertexInstance = MeshDescription->GetVertexInstance(VertexInstanceID);
+		//Copy the tangents
+		for (int32 i = 0; i < 3; ++i)
+		{
+			VertexInstance.Tangent[i] = Tangent[i];
+		}
+		//Copy the bi normal sign
+		VertexInstance.BinormalSign = BitangentSign;
+	}
+
+	void MikkGetTexCoord(const SMikkTSpaceContext* Context, float UV[2], const int FaceIdx, const int VertIdx)
+	{
+		UMeshDescription *MeshDescription = (UMeshDescription*)(Context->m_pUserData);
+		const FVector2D &TexCoord = MeshDescription->VertexInstanceAttributes().GetAttribute<FVector2D>(FVertexInstanceID(FaceIdx * 3 + VertIdx), UEditableMeshAttribute::VertexTextureCoordinate(), 0);
+		UV[0] = TexCoord.X;
+		UV[1] = TexCoord.Y;
+	}
+}
\ No newline at end of file
diff --git a/Engine/Source/Runtime/MeshBuilder/Private/MeshDescriptionHelper.h b/Engine/Source/Runtime/MeshBuilder/Private/MeshDescriptionHelper.h
new file mode 100644
--- /dev/null
+++ b/Engine/Source/Runtime/MeshBuilder/Private/MeshDescriptionHelper.h
@@ -0,0 +1,102 @@
+// Copyright 1998-2017 Epic Games, Inc. All Rights Reserved.
+#pragma once
+
+#include "CoreMinimal.h"
+#include "CoreTypes.h"
+
+class UObject;
+class UMeshDescription;
+struct FMeshBuildSettings;
+struct FVertexInstanceID;
+
+class FMeshDescriptionHelper
+{
+public:
+	enum ETangentOptions
+	{
+		None = 0,
+		BlendOverlappingNormals = 0x1,
+		IgnoreDegenerateTriangles = 0x2,
+		UseMikkTSpace = 0x4,
+	};
+
+	enum class ELightmapUVVersion : int32
+	{
+		BitByBit = 0,
+		Segments = 1,
+		SmallChartPacking = 2,
+		Latest = SmallChartPacking
+	};
+
+	FMeshDescriptionHelper(FMeshBuildSettings* InBuildSettings, const UMeshDescription* InOriginalMeshDescription);
+
+	//Build a render mesh description with the BuildSettings. This will update the InRenderMeshDescription ptr content
+	UMeshDescription* GetRenderMeshDescription(UObject* Owner);
+
+	void CopyMeshDescription(UMeshDescription* SourceMeshDescription, UMeshDescription* DestinationMeshDescription) const;
+
+	//Return true if there is a valid original mesh description, false otherwise(Auto generate LOD).
+	bool IsValidOriginalMeshDescription();
+
+	const TMultiMap<int32, int32>& GetOverlappingCorners() const { return OverlappingCorners; }
+
+private:
+
+	//////////////////////////////////////////////////////////////////////////
+	//PRIVATE function declarations
+
+	const FVector& GetVertexPositionFromVertexInstance(UMeshDescription* MeshDescription, int32 VertexInstanceIndex) const;
+	FVector2D GetVertexInstanceUV(UMeshDescription* MeshDescription, int32 VertexInstanceIndex, int32 UVLayer) const;
+	void GetVertexInstanceNTB(UMeshDescription* MeshDescription, const FVertexInstanceID& VertexInstanceID, FVector &OutNormal, FVector &OutTangent, FVector &OutBiNormal) const;
+	void SetVertexInstanceNTB(UMeshDescription* MeshDescription, const FVertexInstanceID& VertexInstanceID, FVector &OutNormal, FVector &OutTangent, FVector &OutBiNormal);
+	void ComputeNTB_MikkTSpace(UMeshDescription* MeshDescription, ETangentOptions TangentOptions);
+	void ComputeTriangleTangents(UMeshDescription* MeshDescription, TArray<FVector>& OutTangentX, TArray<FVector>& OutTangentY, TArray<FVector>& OutTangentZ, float ComparisonThreshold);
+	void FindOverlappingCorners(float ComparisonThreshold);
+
+	//////////////////////////////////////////////////////////////////////////
+	//PRIVATE class members
+
+	const UMeshDescription* OriginalMeshDescription;
+	FMeshBuildSettings *BuildSettings;
+	TMultiMap<int32, int32> OverlappingCorners;
+
+	
+	//////////////////////////////////////////////////////////////////////////
+	//INLINE small helper use to optimize search and compare
+
+	/** Helper struct for building acceleration structures. */
+	struct FIndexAndZ
+	{
+		float Z;
+		int32 Index;
+
+		/** Default constructor. */
+		FIndexAndZ() {}
+
+		/** Initialization constructor. */
+		FIndexAndZ(int32 InIndex, FVector V)
+		{
+			Z = 0.30f * V.X + 0.33f * V.Y + 0.37f * V.Z;
+			Index = InIndex;
+		}
+	};
+
+	/** Sorting function for vertex Z/index pairs. */
+	struct FCompareIndexAndZ
+	{
+		FORCEINLINE bool operator()(FIndexAndZ const& A, FIndexAndZ const& B) const { return A.Z < B.Z; }
+	};
+
+	/**
+	* Smoothing group interpretation helper structure.
+	*/
+	struct FFanFace
+	{
+		int32 FaceIndex;
+		int32 LinkedVertexIndex;
+		bool bFilled;
+		bool bBlendTangents;
+		bool bBlendNormals;
+	};
+
+};
diff --git a/Engine/Source/Runtime/MeshBuilder/Private/StaticMeshBuilder.cpp b/Engine/Source/Runtime/MeshBuilder/Private/StaticMeshBuilder.cpp
--- a/Engine/Source/Runtime/MeshBuilder/Private/StaticMeshBuilder.cpp
+++ b/Engine/Source/Runtime/MeshBuilder/Private/StaticMeshBuilder.cpp
@@ -5,8 +5,12 @@
 #include "StaticMeshResources.h"
 #include "PhysicsEngine/BodySetup.h"
 #include "MeshDescription.h"
+#include "MeshDescriptionHelper.h"
 #include "BuildOptimizationHelper.h"
 #include "Components.h"
+#include "IMeshReductionManagerModule.h"
+#include "MeshBuild.h"
+
 
 //////////////////////////////////////////////////////////////////////////
 //Local functions definition
@@ -14,8 +18,8 @@ int32 GetPolygonGroupTriangles(const class UMeshDescription *MeshDescription, TA
 bool IsOrphanedVertex(const class UMeshDescription *MeshDescription, const FVertexID VertexID);
 void UpdateBounds(class UStaticMesh* StaticMesh);
 void UpdateCollision(UStaticMesh *StaticMesh);
-void BuildVertexBuffer(const class UMeshDescription* MeshDescription, struct FStaticMeshLODResources& StaticMeshLOD, const struct FMeshBuildSettings& LODBuildSettings, TArray< FStaticMeshBuildVertex >& StaticMeshBuildVertices);
-void BuildIndexBuffer(class UStaticMesh *StaticMesh, const class UMeshDescription* MeshDescription, TArray< uint32 >& IndexBuffer, struct FStaticMeshLODResources& StaticMeshLOD, TArray<int32>& OutWedgeMap, TArray<TArray<uint32> >& OutPerSectionIndices);
+void BuildVertexBuffer(const class UMeshDescription* MeshDescription, struct FStaticMeshLODResources& StaticMeshLOD, const struct FMeshBuildSettings& LODBuildSettings, TArray< FStaticMeshBuildVertex >& StaticMeshBuildVertices, const TMultiMap<int32, int32>& OverlappingCorners, float VertexComparisonThreshold, TArray<int32>& RemapVerts);
+void BuildIndexBuffer(class UStaticMesh *StaticMesh, const class UMeshDescription* MeshDescription, TArray< uint32 >& IndexBuffer, struct FStaticMeshLODResources& StaticMeshLOD, TArray<int32>& OutWedgeMap, TArray<TArray<uint32> >& OutPerSectionIndices, TArray<int32>& RemapVerts);
 void BuildAllBufferOptimizations(struct FStaticMeshLODResources& StaticMeshLOD, const struct FMeshBuildSettings& LODBuildSettings, TArray< uint32 >& IndexBuffer, bool bNeeds32BitIndices, TArray< FStaticMeshBuildVertex >& StaticMeshBuildVertices);
 //////////////////////////////////////////////////////////////////////////
 
@@ -26,7 +30,7 @@ FStaticMeshBuilder::FStaticMeshBuilder()
 
 bool FStaticMeshBuilder::Build(UStaticMesh* StaticMesh)
 {
-	int32 MeshDescriptionCount = StaticMesh->GetMeshDescriptionCount();
+	int32 MeshDescriptionCount = StaticMesh->GetOriginalMeshDescriptionCount();
 	if (MeshDescriptionCount <= 0)
 	{
 		//TODO: Warn the user that there is no mesh description data
@@ -39,14 +43,17 @@ bool FStaticMeshBuilder::Build(UStaticMesh* StaticMesh)
 	FStaticMeshRenderData& StaticMeshRenderData = *StaticMesh->RenderData;
 	for (int32 LodIndex = 0; LodIndex < MeshDescriptionCount; ++LodIndex)
 	{
-		const UMeshDescription* MeshDescription = StaticMesh->GetMeshDescription(LodIndex);
+		FMeshBuildSettings& LODBuildSettings = StaticMesh->SourceModels[LodIndex].BuildSettings;
+		const UMeshDescription* OriginalMeshDescription = StaticMesh->GetOriginalMeshDescription(LodIndex);
+		FMeshDescriptionHelper MeshDescriptionHelper(&LODBuildSettings, OriginalMeshDescription);
+		UMeshDescription* MeshDescription = MeshDescriptionHelper.GetRenderMeshDescription(StaticMesh);
 		check(MeshDescription != nullptr);
+		StaticMesh->SetMeshDescription(LodIndex, MeshDescription);
 		
 		const FPolygonGroupArray& PolygonGroups = MeshDescription->PolygonGroups();
 
 		FStaticMeshLODResources& StaticMeshLOD = StaticMeshRenderData.LODResources[LodIndex];
-		const FMeshBuildSettings& LODBuildSettings = StaticMesh->SourceModels[LodIndex].BuildSettings;
-		
+
 		//TODO: discover degenerate triangle with this threshold
 		float VertexComparisonThreshold = LODBuildSettings.bRemoveDegenerates ? THRESH_POINTS_ARE_SAME : 0.0f;
 
@@ -58,8 +65,8 @@ bool FStaticMeshBuilder::Build(UStaticMesh* StaticMesh)
 		IndexBuffer.Reset();
 
 		StaticMeshLOD.Sections.Empty(PolygonGroups.Num());
-
-		BuildVertexBuffer(MeshDescription, StaticMeshLOD, LODBuildSettings, StaticMeshBuildVertices);
+		TArray<int32> RemapVerts; //Because we will remove MeshVertex that are redundant, we need a remap
+		BuildVertexBuffer(MeshDescription, StaticMeshLOD, LODBuildSettings, StaticMeshBuildVertices, MeshDescriptionHelper.GetOverlappingCorners(), VertexComparisonThreshold, RemapVerts);
 
 		//Render data Wedge map is only set for LOD 0???
 		TArray<int32> TempWedgeMap;
@@ -77,7 +84,7 @@ bool FStaticMeshBuilder::Build(UStaticMesh* StaticMesh)
 			PerSectionIndices.Push(TArray<uint32>());
 		}
 
-		BuildIndexBuffer(StaticMesh, MeshDescription, IndexBuffer, StaticMeshLOD, WedgeMap, PerSectionIndices);
+		BuildIndexBuffer(StaticMesh, MeshDescription, IndexBuffer, StaticMeshLOD, WedgeMap, PerSectionIndices, RemapVerts);
 
 		// Figure out which index buffer stride we need
 		bool bNeeds32BitIndices = false;
@@ -126,8 +133,6 @@ void FStaticMeshBuilder::OnBuildRenderMeshStart(UStaticMesh* StaticMesh, const b
 
 void FStaticMeshBuilder::OnBuildRenderMeshFinish(UStaticMesh* StaticMesh, const bool bRebuildBoundsAndCollision)
 {
-
-
 	if (bRebuildBoundsAndCollision)
 	{
 		UpdateBounds(StaticMesh);
@@ -280,33 +285,108 @@ void UpdateCollision(UStaticMesh *StaticMesh)
 	}
 }
 
-void BuildVertexBuffer(const UMeshDescription* MeshDescription, FStaticMeshLODResources& StaticMeshLOD, const FMeshBuildSettings& LODBuildSettings, TArray< FStaticMeshBuildVertex >& StaticMeshBuildVertices)
+bool AreVerticesEqual(FStaticMeshBuildVertex const& A, FStaticMeshBuildVertex const& B, float ComparisonThreshold)
+{
+	if (   !A.Position.Equals(B.Position, ComparisonThreshold)
+		|| !NormalsEqual(A.TangentX, B.TangentX)
+		|| !NormalsEqual(A.TangentY, B.TangentY)
+		|| !NormalsEqual(A.TangentZ, B.TangentZ)
+		|| A.Color != B.Color)
+	{
+		return false;
+	}
+
+	// UVs
+	for (int32 UVIndex = 0; UVIndex < MAX_STATIC_TEXCOORDS; UVIndex++)
+	{
+		if (!UVsEqual(A.UVs[UVIndex], B.UVs[UVIndex]))
+		{
+			return false;
+		}
+	}
+
+	return true;
+}
+
+void BuildVertexBuffer(const UMeshDescription* MeshDescription, FStaticMeshLODResources& StaticMeshLOD, const FMeshBuildSettings& LODBuildSettings, TArray< FStaticMeshBuildVertex >& StaticMeshBuildVertices, const TMultiMap<int32, int32>& OverlappingCorners, float VertexComparisonThreshold, TArray<int32>& RemapVerts)
 {
 	const FVertexArray& Vertices = MeshDescription->Vertices();
 	const FVertexInstanceArray& VertexInstances = MeshDescription->VertexInstances();
 
+	TArray<int32> RemapVertexInstanceID;
 	// set up vertex buffer elements
-	StaticMeshBuildVertices.SetNum(VertexInstances.GetArraySize());
+	StaticMeshBuildVertices.Reserve(VertexInstances.Num());
 	bool bHasColor = false;
-	for (const FVertexInstanceID VertexInstanceID : VertexInstances.GetElementIDs())
+	//Fill the remap array
+	RemapVerts.AddZeroed(VertexInstances.Num());
+	for (int32& RemapIndex : RemapVerts)
 	{
-		const FMeshVertexInstance& VertexInstance = VertexInstances[VertexInstanceID];
-
-		if (VertexInstance.Color != FColor::White)
-		{
-			bHasColor = true;
-		}
+		RemapIndex = INDEX_NONE;
+	}
+	TArray<int32> DupVerts;
 
-		FStaticMeshBuildVertex& StaticMeshVertex = StaticMeshBuildVertices[VertexInstanceID.GetValue()];
+	const FPolygonGroupArray& PolygonGroups = MeshDescription->PolygonGroups();
+	// Set up index buffer
+	for (const FPolygonGroupID& PolygonGroupID : PolygonGroups.GetElementIDs())
+	{
+		TArray<FMeshTriangle> Triangles;
+		int32 SectionTriangleCount = GetPolygonGroupTriangles(MeshDescription, Triangles, PolygonGroupID);
 
-		StaticMeshVertex.Position = Vertices[VertexInstance.VertexID].VertexPosition * LODBuildSettings.BuildScale3D;
-		StaticMeshVertex.TangentX = VertexInstance.Tangent;
-		StaticMeshVertex.TangentY = FVector::CrossProduct(VertexInstance.Normal, VertexInstance.Tangent).GetSafeNormal() * VertexInstance.BinormalSign;
-		StaticMeshVertex.TangentZ = VertexInstance.Normal;
-		StaticMeshVertex.Color = VertexInstance.Color.ToFColor(true);
-		for (int32 UVIndex = 0; UVIndex < VertexInstance.VertexUVs.Num(); ++UVIndex)
+		const FMeshPolygonGroup& PolygonGroup = PolygonGroups[PolygonGroupID];
+		if (Triangles.Num() > 0)
 		{
-			StaticMeshVertex.UVs[UVIndex] = VertexInstance.VertexUVs[UVIndex];
+			for (int32 TriangleIndex = 0; TriangleIndex < Triangles.Num(); ++TriangleIndex)
+			{
+				const FMeshTriangle& Triangle = Triangles[TriangleIndex];
+				for (int32 TriVert = 0; TriVert < 3; ++TriVert)
+				{
+					int32 VertexInstanceValue = Triangle.GetVertexInstanceID(TriVert).GetValue();
+					check((TriangleIndex * 3) + TriVert == VertexInstanceValue);
+					const FMeshVertexInstance& VertexInstance = MeshDescription->GetVertexInstance(Triangle.GetVertexInstanceID(TriVert));
+					if (VertexInstance.Color != FColor::White)
+					{
+						bHasColor = true;
+					}
+
+					FStaticMeshBuildVertex StaticMeshVertex;
+
+					StaticMeshVertex.Position = Vertices[VertexInstance.VertexID].VertexPosition * LODBuildSettings.BuildScale3D;
+					StaticMeshVertex.TangentX = VertexInstance.Tangent;
+					StaticMeshVertex.TangentY = FVector::CrossProduct(VertexInstance.Normal, VertexInstance.Tangent).GetSafeNormal() * VertexInstance.BinormalSign;
+					StaticMeshVertex.TangentZ = VertexInstance.Normal;
+					StaticMeshVertex.Color = VertexInstance.Color.ToFColor(true);
+					for (int32 UVIndex = 0; UVIndex < VertexInstance.VertexUVs.Num(); ++UVIndex)
+					{
+						StaticMeshVertex.UVs[UVIndex] = VertexInstance.VertexUVs[UVIndex];
+					}
+
+					//Never add duplicated vertex instance
+					DupVerts.Reset();
+					OverlappingCorners.MultiFind(VertexInstanceValue, DupVerts);
+					DupVerts.Sort();
+					int32 Index = INDEX_NONE;
+					for (int32 k = 0; k < DupVerts.Num(); k++)
+					{
+						if (DupVerts[k] >= VertexInstanceValue)
+						{
+							// the verts beyond me haven't been placed yet, so these duplicates are not relevant
+							break;
+						}
+
+						int32 Location = RemapVerts.IsValidIndex(DupVerts[k]) ? RemapVerts[DupVerts[k]] : INDEX_NONE;
+						if (Location != INDEX_NONE && AreVerticesEqual(StaticMeshVertex, StaticMeshBuildVertices[Location], VertexComparisonThreshold))
+						{
+							Index = Location;
+							break;
+						}
+					}
+					if (Index == INDEX_NONE)
+					{
+						Index = StaticMeshBuildVertices.Add(StaticMeshVertex);
+					}
+					RemapVerts[VertexInstanceValue] = Index;
+				}
+			}
 		}
 	}
 
@@ -324,7 +404,7 @@ void BuildVertexBuffer(const UMeshDescription* MeshDescription, FStaticMeshLODRe
 	}
 }
 
-void BuildIndexBuffer(UStaticMesh *StaticMesh, const UMeshDescription* MeshDescription, TArray< uint32 >& IndexBuffer, FStaticMeshLODResources& StaticMeshLOD, TArray<int32>& OutWedgeMap, TArray<TArray<uint32> >& OutPerSectionIndices)
+void BuildIndexBuffer(UStaticMesh *StaticMesh, const UMeshDescription* MeshDescription, TArray< uint32 >& IndexBuffer, FStaticMeshLODResources& StaticMeshLOD, TArray<int32>& OutWedgeMap, TArray<TArray<uint32> >& OutPerSectionIndices, TArray<int32>& RemapVerts)
 {
 	OutWedgeMap.Reset();
 
@@ -364,7 +444,7 @@ void BuildIndexBuffer(UStaticMesh *StaticMesh, const UMeshDescription* MeshDescr
 				const FMeshTriangle& Triangle = Triangles[TriangleIndex];
 				for (int32 TriVert = 0; TriVert < 3; ++TriVert)
 				{
-					const uint32 RenderingVertexIndex = Triangle.GetVertexInstanceID(TriVert).GetValue();
+					const uint32 RenderingVertexIndex = RemapVerts[Triangle.GetVertexInstanceID(TriVert).GetValue()];
 					IndexBuffer.Add(RenderingVertexIndex);
 					OutWedgeMap.Add(RenderingVertexIndex);
 					SectionIndices.Add(RenderingVertexIndex);
