Rewrite the Create normal and tangent function for mesh description.
Use recursive connected triangles with hard edges support to find the smooth normal for every vertex

diff --git a/Engine/Source/Editor/UnrealEd/Private/Fbx/FbxStaticMeshImport.cpp b/Engine/Source/Editor/UnrealEd/Private/Fbx/FbxStaticMeshImport.cpp
--- a/Engine/Source/Editor/UnrealEd/Private/Fbx/FbxStaticMeshImport.cpp
+++ b/Engine/Source/Editor/UnrealEd/Private/Fbx/FbxStaticMeshImport.cpp
@@ -593,7 +593,6 @@ bool UnFbx::FFbxImporter::BuildStaticMeshFromGeometry(FbxNode* Node, UStaticMesh
 						FinalUVVector.X = static_cast<float>(UVVector[0]);
 						FinalUVVector.Y = 1.f - static_cast<float>(UVVector[1]);   //flip the Y of UVs for DirectX
 					}
-					MeshDescription->VertexInstanceAttributes().SetAttribute<FVector2D>(AddedVertexInstanceId, UEditableMeshAttribute::VertexTextureCoordinate(), UVLayerIndex, FinalUVVector);
 					NewVertexInstance.VertexUVs.Add(FinalUVVector);
 				}
 
@@ -616,7 +615,6 @@ bool UnFbx::FFbxImporter::BuildStaticMeshFromGeometry(FbxNode* Node, UStaticMesh
 							uint8(255.f*VertexColor.mBlue),
 							uint8(255.f*VertexColor.mAlpha)
 						);
-						MeshDescription->VertexInstanceAttributes().SetAttribute<FVector4>(AddedVertexInstanceId, UEditableMeshAttribute::VertexColor(), 0, FVector4(FLinearColor(VertexInstanceColor)));
 						NewVertexInstance.Color = FLinearColor(VertexInstanceColor);
 					}
 				}
@@ -628,7 +626,6 @@ bool UnFbx::FFbxImporter::BuildStaticMeshFromGeometry(FbxNode* Node, UStaticMesh
 					if (PaintedColor)
 					{
 						// A matching color for this vertex was found
-						MeshDescription->VertexInstanceAttributes().SetAttribute<FVector4>(AddedVertexInstanceId, UEditableMeshAttribute::VertexColor(), 0, FVector4(FLinearColor(*PaintedColor)));
 						NewVertexInstance.Color = FLinearColor(*PaintedColor);
 					}
 				}
@@ -636,7 +633,6 @@ bool UnFbx::FFbxImporter::BuildStaticMeshFromGeometry(FbxNode* Node, UStaticMesh
 				{
 					// set the triangle's vertex color to a constant override
 					check(VertexColorImportOption == EVertexColorImportOption::Override);
-					MeshDescription->VertexInstanceAttributes().SetAttribute<FVector4>(AddedVertexInstanceId, UEditableMeshAttribute::VertexColor(), 0, FVector4(FLinearColor(VertexOverrideColor)));
 					NewVertexInstance.Color = FLinearColor(VertexOverrideColor);
 				}
 
@@ -652,7 +648,6 @@ bool UnFbx::FFbxImporter::BuildStaticMeshFromGeometry(FbxNode* Node, UStaticMesh
 					FbxVector4 TempValue = LayerElementNormal->GetDirectArray().GetAt(NormalValueIndex);
 					TempValue = TotalMatrixForNormal.MultT(TempValue);
 					FVector TangentZ = Converter.ConvertDir(TempValue);
-					MeshDescription->VertexInstanceAttributes().SetAttribute<FVector>(AddedVertexInstanceId, UEditableMeshAttribute::VertexNormal(), 0, TangentZ.GetSafeNormal());
 					NewVertexInstance.Normal = TangentZ.GetSafeNormal();
 					//tangents and binormals share the same reference, mapping mode and index array
 					if (bHasNTBInformation)
@@ -665,7 +660,6 @@ bool UnFbx::FFbxImporter::BuildStaticMeshFromGeometry(FbxNode* Node, UStaticMesh
 						TempValue = LayerElementTangent->GetDirectArray().GetAt(TangentValueIndex);
 						TempValue = TotalMatrixForNormal.MultT(TempValue);
 						FVector TangentX = Converter.ConvertDir(TempValue);
-						MeshDescription->VertexInstanceAttributes().SetAttribute<FVector>(AddedVertexInstanceId, UEditableMeshAttribute::VertexTangent(), 0, TangentX.GetSafeNormal());
 						NewVertexInstance.Tangent = TangentX.GetSafeNormal();
 
 						int BinormalMapIndex = (BinormalMappingMode == FbxLayerElement::eByControlPoint) ?
@@ -676,9 +670,7 @@ bool UnFbx::FFbxImporter::BuildStaticMeshFromGeometry(FbxNode* Node, UStaticMesh
 						TempValue = LayerElementBinormal->GetDirectArray().GetAt(BinormalValueIndex);
 						TempValue = TotalMatrixForNormal.MultT(TempValue);
 						FVector TangentY = -Converter.ConvertDir(TempValue);
-						FVector BiNormal = TangentY.GetSafeNormal();
-						MeshDescription->VertexInstanceAttributes().SetAttribute<float>(AddedVertexInstanceId, UEditableMeshAttribute::VertexBinormalSign(), 0, GetBasisDeterminantSign(TangentX.GetSafeNormal(), BiNormal, TangentZ.GetSafeNormal()) );
-						NewVertexInstance.BinormalSign = GetBasisDeterminantSign(TangentX.GetSafeNormal(), BiNormal, TangentZ.GetSafeNormal());
+						NewVertexInstance.BinormalSign = GetBasisDeterminantSign(TangentX.GetSafeNormal(), TangentY.GetSafeNormal(), TangentZ.GetSafeNormal());
 					}
 				}
 			}
@@ -760,12 +752,11 @@ bool UnFbx::FFbxImporter::BuildStaticMeshFromGeometry(FbxNode* Node, UStaticMesh
 			// Create polygon edges
 			TArray<UMeshDescription::FContourPoint> Contours;
 			{
-				TArray<FEdgeID> EdgeIDs;
-				TArray<bool> EdgesIsHard;
 				// Add the edges of this triangle
 				for (uint32 TriangleEdgeNumber = 0; TriangleEdgeNumber < 3; ++TriangleEdgeNumber)
 				{
-					UMeshDescription::FContourPoint ContourPoint;
+					int32 ContourPointIndex = Contours.AddDefaulted();
+					UMeshDescription::FContourPoint& ContourPoint = Contours[ContourPointIndex];
 					//Find the matching edge ID
 					uint32 CornerIndices[2];
 					CornerIndices[0] = (TriangleEdgeNumber + 0) % 3;
@@ -774,29 +765,25 @@ bool UnFbx::FFbxImporter::BuildStaticMeshFromGeometry(FbxNode* Node, UStaticMesh
 					FVertexID EdgeVertexIDs[2];
 					EdgeVertexIDs[0] = CornerVerticesIDs[CornerIndices[0]];
 					EdgeVertexIDs[1] = CornerVerticesIDs[CornerIndices[1]];
-					
 
 					FEdgeID MatchEdgeId = FEdgeID::Invalid;
 					const FMeshVertex& EdgeMeshVertex = MeshDescription->GetVertex(EdgeVertexIDs[0]);
 					for (const FEdgeID& EdgeID : EdgeMeshVertex.ConnectedEdgeIDs)
 					{
-						FVertexID OutEdgeVertexID0 = MeshDescription->Edges()[EdgeID].VertexIDs[0];
-						FVertexID OutEdgeVertexID1 = MeshDescription->Edges()[EdgeID].VertexIDs[1];
-						if (EdgeVertexIDs[0] == OutEdgeVertexID0 && EdgeVertexIDs[1] == OutEdgeVertexID1 ||
-							EdgeVertexIDs[1] == OutEdgeVertexID0 && EdgeVertexIDs[0] == OutEdgeVertexID1)
+						FMeshEdge& Edge = MeshDescription->GetEdge(EdgeID);
+						if (EdgeVertexIDs[0] == Edge.VertexIDs[0] && EdgeVertexIDs[1] == Edge.VertexIDs[1] ||
+							EdgeVertexIDs[1] == Edge.VertexIDs[0] && EdgeVertexIDs[0] == Edge.VertexIDs[1])
 						{
 							MatchEdgeId = EdgeID;
 							break;
 						}
 					}
-					int32 SmoothMaskValue = 0;
 					if (MatchEdgeId == FEdgeID::Invalid)
 					{
 						MatchEdgeId = MeshDescription->CreateEdge(EdgeVertexIDs[0], EdgeVertexIDs[1]);
 					}
 					ContourPoint.EdgeID = MatchEdgeId;
 					ContourPoint.VertexInstanceID = CornerInstanceIDs[CornerIndices[0]];
-					Contours.Add(ContourPoint);
 					FMeshEdge& ExistingEdge = MeshDescription->GetEdge(MatchEdgeId);
 
 					int32 EdgeIndex = Mesh->GetMeshEdgeIndexForPolygon(TriangleIndex, TriangleEdgeNumber);
@@ -806,14 +793,8 @@ bool UnFbx::FFbxImporter::BuildStaticMeshFromGeometry(FbxNode* Node, UStaticMesh
 						if (SmoothingMappingMode == FbxLayerElement::eByEdge)
 						{
 							int32 lSmoothingIndex = (SmoothingReferenceMode == FbxLayerElement::eDirect) ? EdgeIndex : SmoothingInfo->GetIndexArray().GetAt(EdgeIndex);
-							SmoothMaskValue = SmoothingInfo->GetDirectArray().GetAt(lSmoothingIndex);
-							if (SmoothMaskValue > 0)
-							{
-								//TODO: not sure I'm doing the right thing here? Maybe I should use eByPolygon smoothing to detect if I have a hard edge by comparing the connected polygons smooth group
-								//Set the hard edges
-								EdgeIDs.Add(MatchEdgeId);
-								EdgesIsHard.Add(true);
-							}
+							//Set the hard edges
+							ExistingEdge.bIsHardEdge = (SmoothingInfo->GetDirectArray().GetAt(lSmoothingIndex) == 0);
 						}
 						else
 						{
@@ -821,11 +802,6 @@ bool UnFbx::FFbxImporter::BuildStaticMeshFromGeometry(FbxNode* Node, UStaticMesh
 						}
 					}
 				}
-				for (int32 EdgeIndex = 0; EdgeIndex < EdgeIDs.Num(); ++EdgeIndex)
-				{
-					FEdgeID EdgeID = EdgeIDs[EdgeIndex];
-					MeshDescription->GetEdge(EdgeID).bIsHardEdge = EdgesIsHard[EdgeIndex];
-				}
 			}
 
 			FPolygonGroupID PolygonGroupID = PolygonGroupMapping[RealMaterialIndex];
@@ -835,15 +811,13 @@ bool UnFbx::FFbxImporter::BuildStaticMeshFromGeometry(FbxNode* Node, UStaticMesh
 			FMeshPolygon& NewPolygon = MeshDescription->GetPolygon(NewPolygonID);
 			NewPolygon.PolygonGroupID = PolygonGroupID;
 			MeshPolygonGroup.Polygons.Add(NewPolygonID);
-			FMeshTriangle NewTriangle;
+			int32 NewTriangleIndex = NewPolygon.Triangles.AddDefaulted();
+			FMeshTriangle& NewTriangle = NewPolygon.Triangles[NewTriangleIndex];
 			for (int32 TriangleVertexIndex = 0; TriangleVertexIndex < 3; ++TriangleVertexIndex)
 			{
 				const FVertexInstanceID &VertexInstanceID = CornerInstanceIDs[TriangleVertexIndex];
 				NewTriangle.SetVertexInstanceID(TriangleVertexIndex, VertexInstanceID);
 			}
-			NewPolygon.Triangles.Add(NewTriangle);
-
-			
 		}
 		//Call this after all GetMeshEdgeIndexForPolygon call this is for optimization purpose.
 		Mesh->EndGetMeshEdgeIndexForPolygon();
diff --git a/Engine/Source/Runtime/MeshBuilder/Private/BuildOptimizationHelper.h b/Engine/Source/Runtime/MeshBuilder/Private/BuildOptimizationHelper.h
--- a/Engine/Source/Runtime/MeshBuilder/Private/BuildOptimizationHelper.h
+++ b/Engine/Source/Runtime/MeshBuilder/Private/BuildOptimizationHelper.h
@@ -3,189 +3,9 @@
 
 #include "CoreMinimal.h"
 #include "CoreTypes.h"
-#include "NvTriStrip.h"
-#include "forsythtriangleorderoptimizer.h"
-#include "nvtess.h"
 #include "Components.h"
 #include "StaticMeshResources.h"
-
-namespace BuildOptimizationPrivate
-{
-	// CVars
-	static TAutoConsoleVariable<int32> CVarTriangleOrderOptimization(
-		TEXT("r.TriangleOrderOptimization"),
-		1,
-		TEXT("Controls the algorithm to use when optimizing the triangle order for the post-transform cache.\n")
-		TEXT("0: Use NVTriStrip (slower)\n")
-		TEXT("1: Use Forsyth algorithm (fastest)(default)")
-		TEXT("2: No triangle order optimization. (least efficient, debugging purposes only)"),
-		ECVF_Default);
-
-	static TAutoConsoleVariable<int32> CVarSupportDepthOnlyIndexBuffers(
-		TEXT("r.SupportDepthOnlyIndexBuffers"),
-		1,
-		TEXT("Enables depth-only index buffers. Saves a little time at the expense of doubling the size of index buffers."),
-		ECVF_ReadOnly | ECVF_RenderThreadSafe);
-
-	static TAutoConsoleVariable<int32> CVarSupportReversedIndexBuffers(
-		TEXT("r.SupportReversedIndexBuffers"),
-		1,
-		TEXT("Enables reversed index buffers. Saves a little time at the expense of doubling the size of index buffers."),
-		ECVF_ReadOnly | ECVF_RenderThreadSafe);
-
-	/*------------------------------------------------------------------------------
-	NVTriStrip for cache optimizing index buffers.
-	------------------------------------------------------------------------------*/
-
-	namespace NvTriStripHelper
-	{
-		/**
-		* Converts 16 bit indices to 32 bit prior to passing them into the real GenerateStrips util method
-		*/
-		void GenerateStrips(
-			const uint8* Indices,
-			bool Is32Bit,
-			const uint32 NumIndices,
-			PrimitiveGroup** PrimGroups,
-			uint32* NumGroups
-		)
-		{
-			if (Is32Bit)
-			{
-				GenerateStrips((uint32*)Indices, NumIndices, PrimGroups, NumGroups);
-			}
-			else
-			{
-				// convert to 32 bit
-				uint32 Idx;
-				TArray<uint32> NewIndices;
-				NewIndices.AddUninitialized(NumIndices);
-				for (Idx = 0; Idx < NumIndices; ++Idx)
-				{
-					NewIndices[Idx] = ((uint16*)Indices)[Idx];
-				}
-				GenerateStrips(NewIndices.GetData(), NumIndices, PrimGroups, NumGroups);
-			}
-
-		}
-
-		/**
-		* Orders a triangle list for better vertex cache coherency.
-		*
-		* *** WARNING: This is safe to call for multiple threads IF AND ONLY IF all
-		* threads call SetListsOnly(true) and SetCacheSize(CACHESIZE_GEFORCE3). If
-		* NvTriStrip is ever used with different settings the library will need
-		* some modifications to be thread-safe. ***
-		*/
-		template<typename IndexDataType, typename Allocator>
-		void CacheOptimizeIndexBuffer(TArray<IndexDataType, Allocator>& Indices)
-		{
-			static_assert(sizeof(IndexDataType) == 2 || sizeof(IndexDataType) == 4, "Indices must be short or int.");
-
-			PrimitiveGroup*	PrimitiveGroups = NULL;
-			uint32			NumPrimitiveGroups = 0;
-			bool Is32Bit = sizeof(IndexDataType) == 4;
-
-			SetListsOnly(true);
-			SetCacheSize(CACHESIZE_GEFORCE3);
-
-			GenerateStrips((uint8*)Indices.GetData(), Is32Bit, Indices.Num(), &PrimitiveGroups, &NumPrimitiveGroups);
-
-			Indices.Empty();
-			Indices.AddUninitialized(PrimitiveGroups->numIndices);
-
-			if (Is32Bit)
-			{
-				FMemory::Memcpy(Indices.GetData(), PrimitiveGroups->indices, Indices.Num() * sizeof(IndexDataType));
-			}
-			else
-			{
-				for (uint32 I = 0; I < PrimitiveGroups->numIndices; ++I)
-				{
-					Indices[I] = (uint16)PrimitiveGroups->indices[I];
-				}
-			}
-
-			delete[] PrimitiveGroups;
-		}
-	}
-
-	/*------------------------------------------------------------------------------
-	Forsyth algorithm for cache optimizing index buffers.
-	------------------------------------------------------------------------------*/
-
-	namespace ForsythHelper
-	{
-		/**
-		* Converts 16 bit indices to 32 bit prior to passing them into the real OptimizeFaces util method
-		*/
-		void OptimizeFaces(
-			const uint8* Indices,
-			bool Is32Bit,
-			const uint32 NumIndices,
-			uint32 NumVertices,
-			uint32* OutIndices,
-			uint16 CacheSize
-		)
-		{
-			if (Is32Bit)
-			{
-				Forsyth::OptimizeFaces((uint32*)Indices, NumIndices, NumVertices, OutIndices, CacheSize);
-			}
-			else
-			{
-				// convert to 32 bit
-				uint32 Idx;
-				TArray<uint32> NewIndices;
-				NewIndices.AddUninitialized(NumIndices);
-				for (Idx = 0; Idx < NumIndices; ++Idx)
-				{
-					NewIndices[Idx] = ((uint16*)Indices)[Idx];
-				}
-				Forsyth::OptimizeFaces(NewIndices.GetData(), NumIndices, NumVertices, OutIndices, CacheSize);
-			}
-
-		}
-
-		/**
-		* Orders a triangle list for better vertex cache coherency.
-		*/
-		template<typename IndexDataType, typename Allocator>
-		void CacheOptimizeIndexBuffer(TArray<IndexDataType, Allocator>& Indices)
-		{
-			static_assert(sizeof(IndexDataType) == 2 || sizeof(IndexDataType) == 4, "Indices must be short or int.");
-			bool Is32Bit = sizeof(IndexDataType) == 4;
-
-			// Count the number of vertices
-			uint32 NumVertices = 0;
-			for (int32 Index = 0; Index < Indices.Num(); ++Index)
-			{
-				if (Indices[Index] > NumVertices)
-				{
-					NumVertices = Indices[Index];
-				}
-			}
-			NumVertices += 1;
-
-			TArray<uint32> OptimizedIndices;
-			OptimizedIndices.AddUninitialized(Indices.Num());
-			uint16 CacheSize = 32;
-			OptimizeFaces((uint8*)Indices.GetData(), Is32Bit, Indices.Num(), NumVertices, OptimizedIndices.GetData(), CacheSize);
-
-			if (Is32Bit)
-			{
-				FMemory::Memcpy(Indices.GetData(), OptimizedIndices.GetData(), Indices.Num() * sizeof(IndexDataType));
-			}
-			else
-			{
-				for (int32 I = 0; I < OptimizedIndices.Num(); ++I)
-				{
-					Indices[I] = (uint16)OptimizedIndices[I];
-				}
-			}
-		}
-	}
-}
+#include "ThirdPartyBuildOptimizationHelper.h"
 
 class BuildOptimizationHelper
 {
@@ -213,20 +33,6 @@ public:
 		FORCEINLINE bool operator()(FIndexAndZ const& A, FIndexAndZ const& B) const { return A.Z < B.Z; }
 	};
 
-	/**
-	* Returns true if the specified points are about equal
-	*/
-	static inline bool PointsEqual(const FVector& V1, const FVector& V2, float ComparisonThreshold)
-	{
-		if (FMath::Abs(V1.X - V2.X) > ComparisonThreshold
-			|| FMath::Abs(V1.Y - V2.Y) > ComparisonThreshold
-			|| FMath::Abs(V1.Z - V2.Z) > ComparisonThreshold)
-		{
-			return false;
-		}
-		return true;
-	}
-
 	struct FMeshEdge
 	{
 		int32	Vertices[2];
@@ -413,65 +219,6 @@ public:
 		}
 	};
 
-	/**
-	* Provides static mesh render data to the NVIDIA tessellation library.
-	*/
-	class FStaticMeshNvRenderBuffer : public nv::RenderBuffer
-	{
-	public:
-
-		/** Construct from static mesh render buffers. */
-		FStaticMeshNvRenderBuffer(
-			const FPositionVertexBuffer& InPositionVertexBuffer,
-			const FStaticMeshVertexBuffer& InVertexBuffer,
-			const TArray<uint32>& Indices)
-			: PositionVertexBuffer(InPositionVertexBuffer)
-			, VertexBuffer(InVertexBuffer)
-		{
-			check(PositionVertexBuffer.GetNumVertices() == VertexBuffer.GetNumVertices());
-			mIb = new nv::IndexBuffer((void*)Indices.GetData(), nv::IBT_U32, Indices.Num(), false);
-		}
-
-		/** Retrieve the position and first texture coordinate of the specified index. */
-		virtual nv::Vertex getVertex(unsigned int Index) const
-		{
-			nv::Vertex Vertex;
-
-			check(Index < PositionVertexBuffer.GetNumVertices());
-
-			const FVector& Position = PositionVertexBuffer.VertexPosition(Index);
-			Vertex.pos.x = Position.X;
-			Vertex.pos.y = Position.Y;
-			Vertex.pos.z = Position.Z;
-
-			if (VertexBuffer.GetNumTexCoords())
-			{
-				const FVector2D UV = VertexBuffer.GetVertexUV(Index, 0);
-				Vertex.uv.x = UV.X;
-				Vertex.uv.y = UV.Y;
-			}
-			else
-			{
-				Vertex.uv.x = 0.0f;
-				Vertex.uv.y = 0.0f;
-			}
-
-			return Vertex;
-		}
-
-	private:
-
-		/** The position vertex buffer for the static mesh. */
-		const FPositionVertexBuffer& PositionVertexBuffer;
-
-		/** The vertex buffer for the static mesh. */
-		const FStaticMeshVertexBuffer& VertexBuffer;
-
-		/** Copying is forbidden. */
-		FStaticMeshNvRenderBuffer(const FStaticMeshNvRenderBuffer&);
-		FStaticMeshNvRenderBuffer& operator=(const FStaticMeshNvRenderBuffer&);
-	};
-
 	static void BuildDepthOnlyIndexBuffer(
 		TArray<uint32>& OutDepthIndices,
 		const TArray<FStaticMeshBuildVertex>& InVertices,
@@ -514,7 +261,7 @@ public:
 					break; // can't be any more dups
 
 				uint32 OtherIndex = VertIndexAndZ[j].Index;
-				if (PointsEqual(InVertices[SrcIndex].Position, InVertices[OtherIndex].Position,/*bUseEpsilonCompare=*/ false))
+				if (InVertices[SrcIndex].Position.Equals(InVertices[OtherIndex].Position, 0.0f))
 				{
 					IndexMap[SrcIndex] = FMath::Min(IndexMap[SrcIndex], OtherIndex);
 					IndexMap[OtherIndex] = FMath::Min(IndexMap[OtherIndex], SrcIndex);
@@ -547,7 +294,7 @@ public:
 	{
 		if (Indices.Num())
 		{
-			FStaticMeshNvRenderBuffer StaticMeshRenderBuffer(PositionVertexBuffer, VertexBuffer, Indices);
+			BuildOptimizationPrivate::NvTriStripHelper::FStaticMeshNvRenderBuffer StaticMeshRenderBuffer(PositionVertexBuffer, VertexBuffer, Indices);
 			nv::IndexBuffer* PnAENIndexBuffer = nv::tess::buildTessellationBuffer(&StaticMeshRenderBuffer, nv::DBM_PnAenDominantCorner, true);
 			check(PnAENIndexBuffer);
 			const int32 IndexCount = (int32)PnAENIndexBuffer->getLength();
diff --git a/Engine/Source/Runtime/MeshBuilder/Private/MeshDescriptionHelper.cpp b/Engine/Source/Runtime/MeshBuilder/Private/MeshDescriptionHelper.cpp
--- a/Engine/Source/Runtime/MeshBuilder/Private/MeshDescriptionHelper.cpp
+++ b/Engine/Source/Runtime/MeshBuilder/Private/MeshDescriptionHelper.cpp
@@ -12,6 +12,11 @@
 #include "Serialization/MemoryReader.h"
 #include "RenderUtils.h"
 
+DEFINE_LOG_CATEGORY(LogMeshBuilder);
+
+//Enable all check
+//#define ENABLE_NTB_CHECK
+
 namespace MeshDescriptionMikktSpaceInterface
 {
 	//Mikk t spce static function
@@ -34,88 +39,87 @@ UMeshDescription* FMeshDescriptionHelper::GetRenderMeshDescription(UObject* Owne
 	//Use the build settings to create the RenderMeshDescription
 	UMeshDescription *RenderMeshDescription = NewObject<UMeshDescription>(Owner, NAME_None, RF_NoFlags);
 
-	if(OriginalMeshDescription != nullptr)
+	if (OriginalMeshDescription == nullptr)
 	{
-		//Copy The Original Mesh Description in the render mesh description
-		CopyMeshDescription(const_cast<UMeshDescription*>(OriginalMeshDescription), RenderMeshDescription);
-		float ComparisonThreshold = BuildSettings->bRemoveDegenerates ? THRESH_POINTS_ARE_SAME : 0.0f;
-		int32 NumWedges = RenderMeshDescription->VertexInstances().Num();
-		// Find overlapping corners to accelerate adjacency.
-		FindOverlappingCorners(ComparisonThreshold);
-
-		// Figure out if we should recompute normals and tangents.
-		int32 OriginalNormalCount = RenderMeshDescription->VertexInstances().Num();
-		int32 OriginalTangentCount = OriginalNormalCount;
-		int32 OriginalBiNormalCount = OriginalNormalCount;
-		bool bRecomputeNormals = BuildSettings->bRecomputeNormals || OriginalNormalCount != NumWedges;
-		bool bRecomputeTangents = BuildSettings->bRecomputeTangents || OriginalTangentCount != NumWedges || OriginalBiNormalCount != NumWedges;
+		//oups, we do not have a valid original mesh to create the render from
+		return RenderMeshDescription;
+	}
+	//Copy The Original Mesh Description in the render mesh description
+	CopyMeshDescription(const_cast<UMeshDescription*>(OriginalMeshDescription), RenderMeshDescription);
+	float ComparisonThreshold = BuildSettings->bRemoveDegenerates ? THRESH_POINTS_ARE_SAME : 0.0f;
 
-		if (bRecomputeTangents || bRecomputeNormals)
+	if (BuildSettings->bRecomputeTangents || BuildSettings->bRecomputeNormals)
+	{
+		for (const FVertexInstanceID& VertexInstanceID : RenderMeshDescription->VertexInstances().GetElementIDs())
 		{
-			for (const FVertexInstanceID& VertexInstanceID : RenderMeshDescription->VertexInstances().GetElementIDs())
+			FMeshVertexInstance& VertexInstance = RenderMeshDescription->GetVertexInstance(VertexInstanceID);
+			// Dump normals and tangents if we are recomputing them.
+			if (BuildSettings->bRecomputeTangents)
 			{
-				FMeshVertexInstance& VertexInstance = RenderMeshDescription->GetVertexInstance(VertexInstanceID);
-				// Dump normals and tangents if we are recomputing them.
-				if (bRecomputeTangents)
-				{
-					//Dump the tangents
-					VertexInstance.BinormalSign = 0.0f;
-					VertexInstance.Tangent = FVector(0.0f);
-				}
-				if (bRecomputeNormals)
-				{
-					//Dump the normals
-					VertexInstance.Normal = FVector(0.0f);
-				}
+				//Dump the tangents
+				VertexInstance.BinormalSign = 0.0f;
+				VertexInstance.Tangent = FVector(0.0f);
+			}
+			if (BuildSettings->bRecomputeNormals)
+			{
+				//Dump the normals
+				VertexInstance.Normal = FVector(0.0f);
 			}
 		}
+	}
 
-		// Compute any missing tangents.
+	// Compute any missing normals or tangents.
+	{
+		// Static meshes always blend normals of overlapping corners.
+		uint32 TangentOptions = FMeshDescriptionHelper::ETangentOptions::BlendOverlappingNormals;
+		if (BuildSettings->bRemoveDegenerates)
 		{
-			// Static meshes always blend normals of overlapping corners.
-			uint32 TangentOptions = FMeshDescriptionHelper::ETangentOptions::BlendOverlappingNormals;
-			if (BuildSettings->bRemoveDegenerates)
-			{
-				// If removing degenerate triangles, ignore them when computing tangents.
-				TangentOptions |= FMeshDescriptionHelper::ETangentOptions::IgnoreDegenerateTriangles;
-			}
+			// If removing degenerate triangles, ignore them when computing tangents.
+			TangentOptions |= FMeshDescriptionHelper::ETangentOptions::IgnoreDegenerateTriangles;
+		}
 
-			//Keep the original mesh description NTBs if we do not rebuild the normals or tangents.
-			//MikkTSpace should be use only when the user want to recompute the normals or tangents otherwise should always fallback on builtin
-			if (BuildSettings->bUseMikkTSpace && (BuildSettings->bRecomputeNormals || BuildSettings->bRecomputeTangents))
-			{
-				ComputeNTB_MikkTSpace(RenderMeshDescription, (FMeshDescriptionHelper::ETangentOptions)TangentOptions);
-			}
-			//TODO: builtin Compute of the tangents
-			//else if(BuildSettings->bRecomputeNormals || BuildSettings->bRecomputeTangents)
-			//{
-				//ComputeTangents(RenderMeshDescription, OverlappingCorners, TangentOptions);
-			//}
+		//This function make sure the Polygon NTB are compute and also remove degenerated triangle from the render mesh description.
+		CreatePolygonNTB(RenderMeshDescription, BuildSettings->bRemoveDegenerates ? SMALL_NUMBER : 0.0f);
+		
+		// Find overlapping corners to accelerate adjacency.
+		FindOverlappingCorners(ComparisonThreshold);
+			
+		//Keep the original mesh description NTBs if we do not rebuild the normals or tangents.
+		bool bComputeTangentLegacy = !BuildSettings->bUseMikkTSpace && (BuildSettings->bRecomputeNormals || BuildSettings->bRecomputeTangents);
+			
+		//MikkTSpace should be use only when the user want to recompute the normals or tangents otherwise should always fallback on builtin
+		//We cannot use mikkt space with degenerated normals fallback on buitin.
+		if (BuildSettings->bUseMikkTSpace && (BuildSettings->bRecomputeNormals || BuildSettings->bRecomputeTangents))
+		{
+			CreateNormals(RenderMeshDescription, (FMeshDescriptionHelper::ETangentOptions)TangentOptions, false);
+			CreateMikktTangents(RenderMeshDescription, (FMeshDescriptionHelper::ETangentOptions)TangentOptions);
+		}
+		else
+		{
+			//Set the compute tangent to true when we do not build using mikkt space
+			CreateNormals(RenderMeshDescription, (FMeshDescriptionHelper::ETangentOptions)TangentOptions, true);
 		}
+	}
 
-		// TODO: Generate lightmap UVs
-		/*if (BuildSettings->bGenerateLightmapUVs)
+	// TODO: Generate lightmap UVs
+	/*if (BuildSettings->bGenerateLightmapUVs)
+	{
+		if (RawMesh.WedgeTexCoords[BuildSettings->SrcLightmapIndex].Num() == 0)
 		{
-			if (RawMesh.WedgeTexCoords[BuildSettings->SrcLightmapIndex].Num() == 0)
-			{
-				BuildSettings->SrcLightmapIndex = 0;
-			}
+			BuildSettings->SrcLightmapIndex = 0;
+		}
 
-			FLayoutUV Packer(&RawMesh, BuildSettings->SrcLightmapIndex, BuildSettings->DstLightmapIndex, BuildSettings->MinLightmapResolution);
-			Packer.SetVersion(LightmapUVVersion);
+		FLayoutUV Packer(&RawMesh, BuildSettings->SrcLightmapIndex, BuildSettings->DstLightmapIndex, BuildSettings->MinLightmapResolution);
+		Packer.SetVersion(LightmapUVVersion);
 
-			Packer.FindCharts(OverlappingCorners);
-			bool bPackSuccess = Packer.FindBestPacking();
-			if (bPackSuccess)
-			{
-				Packer.CommitPackedUVs();
-			}
+		Packer.FindCharts(OverlappingCorners);
+		bool bPackSuccess = Packer.FindBestPacking();
+		if (bPackSuccess)
+		{
+			Packer.CommitPackedUVs();
 		}
-		HasRawMesh[LODIndex] = true;*/
-	}
-	else
-	{
 	}
+	HasRawMesh[LODIndex] = true;*/
 
 	return RenderMeshDescription;
 }
@@ -140,13 +144,12 @@ bool FMeshDescriptionHelper::IsValidOriginalMeshDescription()
 void FMeshDescriptionHelper::FindOverlappingCorners(float ComparisonThreshold)
 {
 	//Empty the old data
-	OverlappingCorners.Empty();
+	OverlappingCorners.Reset();
 	const int32 NumWedges = OriginalMeshDescription->VertexInstances().Num();
 
 	// Create a list of vertex Z/index pairs
 	TArray<FIndexAndZ> VertIndexAndZ;
 	VertIndexAndZ.Reserve(NumWedges);
-	const FVertexArray& Vertices = OriginalMeshDescription->Vertices();
 	const FVertexInstanceArray& VertexInstances = OriginalMeshDescription->VertexInstances();
 
 	for (const FPolygonID& PolygonID : OriginalMeshDescription->Polygons().GetElementIDs())
@@ -157,9 +160,9 @@ void FMeshDescriptionHelper::FindOverlappingCorners(float ComparisonThreshold)
 			for (int32 CornerIndex = 0; CornerIndex < 3; ++CornerIndex)
 			{
 				const FVertexInstanceID& VertexInstanceID = MeshTriangle.GetVertexInstanceID(CornerIndex);
-				const FMeshVertexInstance& VertexInstance = VertexInstances[VertexInstanceID];
+				const FMeshVertexInstance& VertexInstance = OriginalMeshDescription->GetVertexInstance(VertexInstanceID);
 
-				new(VertIndexAndZ)FIndexAndZ(VertexInstanceID.GetValue(), Vertices[VertexInstance.VertexID].VertexPosition);
+				new(VertIndexAndZ)FIndexAndZ(VertexInstanceID.GetValue(), OriginalMeshDescription->GetVertex(VertexInstance.VertexID).VertexPosition);
 			}
 		}
 	}
@@ -188,66 +191,46 @@ void FMeshDescriptionHelper::FindOverlappingCorners(float ComparisonThreshold)
 	}
 }
 
-const FVector& FMeshDescriptionHelper::GetVertexPositionFromVertexInstance(UMeshDescription* MeshDescription, int32 VertexInstanceIndex) const
-{
-	return MeshDescription->GetVertex(MeshDescription->GetVertexInstance(FVertexInstanceID(VertexInstanceIndex)).VertexID).VertexPosition;
-}
-
-FVector2D FMeshDescriptionHelper::GetVertexInstanceUV(UMeshDescription* MeshDescription, int32 VertexInstanceIndex, int32 UVLayer) const
+const FVector& FMeshDescriptionHelper::GetVertexPositionFromVertexInstance(UMeshDescription* MeshDescription, const FVertexInstanceID& VertexInstanceID) const
 {
-	return MeshDescription->VertexInstanceAttributes().GetAttribute<FVector2D>(FVertexInstanceID(VertexInstanceIndex), UEditableMeshAttribute::VertexTextureCoordinate(), UVLayer);
-}
-
-void FMeshDescriptionHelper::GetVertexInstanceNTB(UMeshDescription* MeshDescription, const FVertexInstanceID& VertexInstanceID, FVector &OutNormal, FVector &OutTangent, FVector &OutBiNormal) const
-{
-	const FMeshVertexInstance& VertexInstance = MeshDescription->GetVertexInstance(VertexInstanceID);
-	OutNormal = VertexInstance.Normal;
-	OutTangent = VertexInstance.Tangent;
-	if (!OutNormal.IsNearlyZero() && !OutTangent.IsNearlyZero())
-	{
-		OutBiNormal = FVector::CrossProduct(VertexInstance.Normal, VertexInstance.Tangent).GetSafeNormal() * VertexInstance.BinormalSign;
-	}
-	else
-	{
-		OutBiNormal = FVector(0.0f);
-	}
+	return MeshDescription->GetVertex(MeshDescription->GetVertexInstance(VertexInstanceID).VertexID).VertexPosition;
 }
 
-void FMeshDescriptionHelper::SetVertexInstanceNTB(UMeshDescription* MeshDescription, const FVertexInstanceID& VertexInstanceID, FVector &OutNormal, FVector &OutTangent, FVector &OutBiNormal)
+FVector2D &FMeshDescriptionHelper::GetVertexInstanceUV(UMeshDescription* MeshDescription, const FVertexInstanceID& VertexInstanceID, int32 UVLayer) const
 {
-	FMeshVertexInstance& VertexInstance = MeshDescription->GetVertexInstance(VertexInstanceID);
-	VertexInstance.Normal = OutNormal;
-	VertexInstance.Tangent = OutTangent;
-	VertexInstance.BinormalSign = GetBasisDeterminantSign(OutTangent, OutBiNormal, OutNormal);
+	return MeshDescription->GetVertexInstance(VertexInstanceID).VertexUVs[UVLayer];
 }
 
-void FMeshDescriptionHelper::ComputeTriangleTangents(
-	UMeshDescription* MeshDescription,
-	TArray<FVector>& OutTangentX,
-	TArray<FVector>& OutTangentY,
-	TArray<FVector>& OutTangentZ,
-	float ComparisonThreshold
-)
+void FMeshDescriptionHelper::CreatePolygonNTB(UMeshDescription* MeshDescription, float ComparisonThreshold)
 {
-	const int32 NumTriangles = MeshDescription->Polygons().Num();
-	OutTangentX.Empty(NumTriangles);
-	OutTangentY.Empty(NumTriangles);
-	OutTangentZ.Empty(NumTriangles);
-
+	TArray<FPolygonID> DegeneratePolygons;
 	for (const FPolygonID& PolygonID : MeshDescription->Polygons().GetElementIDs())
 	{
-		const TArray<FMeshTriangle>& MeshTriangles = MeshDescription->GetPolygon(PolygonID).Triangles;
-		for (const FMeshTriangle& MeshTriangle : MeshTriangles)
+		FVector TangentX(0.0f);
+		FVector TangentY(0.0f);
+		FVector TangentZ(0.0f);
+
+		FMeshPolygon& Polygon = MeshDescription->GetPolygon(PolygonID);
+		const TArray<FMeshTriangle>& MeshTriangles = Polygon.Triangles;
+#ifdef ENABLE_NTB_CHECK
+		//Assume triangle are build
+		check(MeshTriangles.Num() > 0);
+#endif
+
+		//We need only the first triangle since all triangle of a polygon must have the same normals (planar polygon)
+		const FMeshTriangle& MeshTriangle = MeshTriangles[0];
+		int32 UVIndex = 0;
+
+		FVector P[3];
+		for (int32 i = 0; i < 3; ++i)
 		{
-			int32 UVIndex = 0;
-
-			FVector P[3];
-			for (int32 i = 0; i < 3; ++i)
-			{
-				P[i] = GetVertexPositionFromVertexInstance(MeshDescription, MeshTriangle.GetVertexInstanceID(i).GetValue());
-			}
+			P[i] = GetVertexPositionFromVertexInstance(MeshDescription, MeshTriangle.GetVertexInstanceID(i));
+		}
 
-			const FVector Normal = ((P[1] - P[2]) ^ (P[0] - P[2])).GetSafeNormal(ComparisonThreshold);
+		const FVector Normal = ((P[1] - P[2]) ^ (P[0] - P[2])).GetSafeNormal(ComparisonThreshold);
+		//Check for degenerated polygons, avoid NAN
+		if (!Normal.IsNearlyZero(ComparisonThreshold))
+		{
 			FMatrix	ParameterToLocal(
 				FPlane(P[1].X - P[0].X, P[1].Y - P[0].Y, P[1].Z - P[0].Z, 0),
 				FPlane(P[2].X - P[0].X, P[2].Y - P[0].Y, P[2].Z - P[0].Z, 0),
@@ -255,9 +238,9 @@ void FMeshDescriptionHelper::ComputeTriangleTangents(
 				FPlane(0, 0, 0, 1)
 			);
 
-			const FVector2D T1 = GetVertexInstanceUV(MeshDescription, MeshTriangle.GetVertexInstanceID(0).GetValue(), 0);
-			const FVector2D T2 = GetVertexInstanceUV(MeshDescription, MeshTriangle.GetVertexInstanceID(1).GetValue(), 0);
-			const FVector2D T3 = GetVertexInstanceUV(MeshDescription, MeshTriangle.GetVertexInstanceID(2).GetValue(), 0);
+			const FVector2D T1 = GetVertexInstanceUV(MeshDescription, MeshTriangle.GetVertexInstanceID(0), 0);
+			const FVector2D T2 = GetVertexInstanceUV(MeshDescription, MeshTriangle.GetVertexInstanceID(1), 0);
+			const FVector2D T3 = GetVertexInstanceUV(MeshDescription, MeshTriangle.GetVertexInstanceID(2), 0);
 
 			FMatrix ParameterToTexture(
 				FPlane(T2.X - T1.X, T2.Y - T1.Y, 0, 0),
@@ -269,317 +252,311 @@ void FMeshDescriptionHelper::ComputeTriangleTangents(
 			// Use InverseSlow to catch singular matrices.  Inverse can miss this sometimes.
 			const FMatrix TextureToLocal = ParameterToTexture.Inverse() * ParameterToLocal;
 
-			OutTangentX.Add(TextureToLocal.TransformVector(FVector(1, 0, 0)).GetSafeNormal());
-			OutTangentY.Add(TextureToLocal.TransformVector(FVector(0, 1, 0)).GetSafeNormal());
-			OutTangentZ.Add(Normal);
-
-			FVector::CreateOrthonormalBasis(
-				OutTangentX[PolygonID.GetValue()],
-				OutTangentY[PolygonID.GetValue()],
-				OutTangentZ[PolygonID.GetValue()]
-			);
+			TangentX = TextureToLocal.TransformVector(FVector(1, 0, 0)).GetSafeNormal();
+			TangentY = TextureToLocal.TransformVector(FVector(0, 1, 0)).GetSafeNormal();
+			TangentZ = Normal;
+			FVector::CreateOrthonormalBasis(TangentX, TangentY, TangentZ);
 		}
-	}
-
-	check(OutTangentX.Num() == NumTriangles);
-	check(OutTangentY.Num() == NumTriangles);
-	check(OutTangentZ.Num() == NumTriangles);
-}
-
-void FMeshDescriptionHelper::ComputeNTB_MikkTSpace(UMeshDescription* MeshDescription, FMeshDescriptionHelper::ETangentOptions TangentOptions)
-{
-	bool bBlendOverlappingNormals = (TangentOptions & FMeshDescriptionHelper::ETangentOptions::BlendOverlappingNormals) != 0;
-	bool bIgnoreDegenerateTriangles = (TangentOptions & FMeshDescriptionHelper::ETangentOptions::IgnoreDegenerateTriangles) != 0;
-	float ComparisonThreshold = bIgnoreDegenerateTriangles ? THRESH_POINTS_ARE_SAME : 0.0f;
-
-	// Compute per-triangle tangents.
-	TArray<FVector> TriangleTangentX;
-	TArray<FVector> TriangleTangentY;
-	TArray<FVector> TriangleTangentZ;
-
-	ComputeTriangleTangents(
-		MeshDescription,
-		TriangleTangentX,
-		TriangleTangentY,
-		TriangleTangentZ,
-		bIgnoreDegenerateTriangles ? SMALL_NUMBER : 0.0f
-	);
-	TArray<uint32> SmoothingGroupIndices;
-	SmoothingGroupIndices.Reserve(MeshDescription->Polygons().Num());
-	
-	//TODO: support smooth group per edges, we currently do not import smooth group
-	for (const FPolygonID& PolygonID : OriginalMeshDescription->Polygons().GetElementIDs())
-	{
-		SmoothingGroupIndices.Add(1);
-	}
-
-	TArray<FVector> OutTangentX;
-	OutTangentX.AddZeroed(MeshDescription->VertexInstances().Num());
-	TArray<FVector> OutTangentY;
-	OutTangentY.AddZeroed(MeshDescription->VertexInstances().Num());
-	TArray<FVector> OutTangentZ;
-	OutTangentZ.AddZeroed(MeshDescription->VertexInstances().Num());
-
-	for (const FPolygonID& PolygonID : MeshDescription->Polygons().GetElementIDs())
-	{
-		const TArray<FMeshTriangle>& MeshTriangles = MeshDescription->GetPolygon(PolygonID).Triangles;
-		for (const FMeshTriangle& MeshTriangle : MeshTriangles)
+		else
 		{
-			for (int32 CornerIndex = 0; CornerIndex < 3; ++CornerIndex)
-			{
-				const FVertexInstanceID& VertexInstanceID = MeshTriangle.GetVertexInstanceID(CornerIndex);
-				GetVertexInstanceNTB(MeshDescription, VertexInstanceID, OutTangentZ[VertexInstanceID.GetValue()], OutTangentX[VertexInstanceID.GetValue()], OutTangentY[VertexInstanceID.GetValue()]);
-			}
+			DegeneratePolygons.Add(PolygonID);
+		}
+		//Set the polygon Normal value only once if it is not set.
+		if (Polygon.PolygonNormal.IsNearlyZero())
+		{
+			Polygon.PolygonTangent = TangentX;
+			Polygon.PolygonBinormal = TangentY;
+			Polygon.PolygonNormal = TangentZ;
 		}
 	}
-	// Declare these out here to avoid reallocations.
-	TArray<FFanFace> RelevantFacesForCorner[3];
-	TArray<int32> AdjacentFaces;
-	TArray<int32> DupVerts;
-
-	int32 NumWedges = MeshDescription->VertexInstances().Num();
-	int32 NumFaces = NumWedges / 3;
-	check(MeshDescription->Polygons().Num() == NumFaces);
-
-	bool bWedgeTSpace = false;
 
-	if (OutTangentX.Num() > 0 && OutTangentY.Num() > 0)
+	//Delete the degenerated polygons. The array is fill only if the remove degenerated option is turn on.
+	if (DegeneratePolygons.Num() > 0)
 	{
-		bWedgeTSpace = true;
-		for (int32 WedgeIdx = 0; bWedgeTSpace && WedgeIdx < OutTangentX.Num() && WedgeIdx < OutTangentY.Num(); ++WedgeIdx)
+		TSet<FEdgeID> OrphanedEdges;
+		TSet<FVertexInstanceID> OrphanedVertexInstances;
+		TSet<FPolygonGroupID> OrphanedPolygonGroups;
+		TSet<FVertexID> OrphanedVertices;
+		for (FPolygonID& PolygonID : DegeneratePolygons)
+		{
+			MeshDescription->DeletePolygon(PolygonID, &OrphanedEdges, &OrphanedVertexInstances, &OrphanedPolygonGroups);
+		}
+		for (FPolygonGroupID& PolygonGroupID : OrphanedPolygonGroups)
+		{
+			MeshDescription->DeletePolygonGroup(PolygonGroupID);
+		}
+		for (FVertexInstanceID& VertexInstanceID : OrphanedVertexInstances)
+		{
+			MeshDescription->DeleteVertexInstance(VertexInstanceID, &OrphanedVertices);
+		}
+		for (FEdgeID& EdgeID : OrphanedEdges)
+		{
+			MeshDescription->DeleteEdge(EdgeID, &OrphanedVertices);
+		}
+		for (FVertexID& VertexID : OrphanedVertices)
 		{
-			bWedgeTSpace &= (!OutTangentX[WedgeIdx].IsNearlyZero()) && (!OutTangentY[WedgeIdx].IsNearlyZero());
+			MeshDescription->DeleteVertex(VertexID);
 		}
+		//Compact and Remap IDs so we have clean ID from 0 to n since we just erase some polygons
+		//The render build need to have compact ID
+		FElementIDRemappings RemappingInfos;
+		MeshDescription->Compact(RemappingInfos);
 	}
+}
 
-	// Allocate storage for tangents if none were provided, and calculate normals for MikkTSpace.
-	if (OutTangentZ.Num() != NumWedges)
+struct FVertexInfo
+{
+	FVertexInfo()
 	{
-		// normals are not included, so we should calculate them
-		OutTangentZ.Empty(NumWedges);
-		OutTangentZ.AddZeroed(NumWedges);
+		PolygonID = FPolygonID::Invalid;
+		VertexInstanceID = FVertexInstanceID::Invalid;
+		UVs = FVector2D(0.0f, 0.0f);
+		EdgeIDs.Reserve(2);//Most of the time a edge has two triangles
 	}
 
-	// we need to calculate normals for MikkTSpace
-	for (const FPolygonID& PolygonID : MeshDescription->Polygons().GetElementIDs())
+	FPolygonID PolygonID;
+	FVertexInstanceID VertexInstanceID;
+	FVector2D UVs;
+	TArray<FEdgeID> EdgeIDs;
+};
+
+void RecursiveFillPolygonGroup(
+	  UMeshDescription* MeshDescription
+	, TMap<FPolygonID, FVertexInfo>& VertexInfoMap
+	, TSet<FPolygonID>& CurrentGroup
+	, FVertexInfo& CurrentVertexInfo
+	, TSet<FPolygonID>& ConsumedPolygon)
+{
+	CurrentGroup.Add(CurrentVertexInfo.PolygonID);
+	ConsumedPolygon.Add(CurrentVertexInfo.PolygonID);
+	for (FEdgeID &EdgeID : CurrentVertexInfo.EdgeIDs)
 	{
-		const TArray<FMeshTriangle>& MeshTriangles = MeshDescription->GetPolygon(PolygonID).Triangles;
-		for (const FMeshTriangle& MeshTriangle : MeshTriangles)
+		FMeshEdge& Edge = MeshDescription->GetEdge(EdgeID);
+		if (Edge.bIsHardEdge)
 		{
-			FVector CornerPositions[3];
-			FVector CornerNormal[3];
-
-			for (int32 CornerIndex = 0; CornerIndex < 3; CornerIndex++)
-			{
-				const FVertexInstanceID& VertexInstanceID = MeshTriangle.GetVertexInstanceID(CornerIndex);
-				CornerNormal[CornerIndex] = FVector::ZeroVector;
-				CornerPositions[CornerIndex] = GetVertexPositionFromVertexInstance(MeshDescription, VertexInstanceID.GetValue());
-				RelevantFacesForCorner[CornerIndex].Reset();
-			}
-
-			// Don't process degenerate triangles.
-			if (CornerPositions[0].Equals(CornerPositions[1], ComparisonThreshold)
-				|| CornerPositions[0].Equals(CornerPositions[2], ComparisonThreshold)
-				|| CornerPositions[1].Equals(CornerPositions[2], ComparisonThreshold))
-			{
-				continue;
-			}
-
-			// No need to process triangles if tangents already exist.
-			bool bCornerHasNormal[3] = { 0 };
-			for (int32 CornerIndex = 0; CornerIndex < 3; CornerIndex++)
-			{
-				const FVertexInstanceID& VertexInstanceID = MeshTriangle.GetVertexInstanceID(CornerIndex);
-				bCornerHasNormal[CornerIndex] = !OutTangentZ[VertexInstanceID.GetValue()].IsNearlyZero();
-			}
-			if (bCornerHasNormal[0] && bCornerHasNormal[1] && bCornerHasNormal[2])
+			//End of the group
+			continue;
+		}
+		for (const FPolygonID& PolygonID : Edge.ConnectedPolygons)
+		{
+			if (PolygonID == CurrentVertexInfo.PolygonID)
 			{
 				continue;
 			}
-
-			// Start building a list of faces adjacent to this face.
-			AdjacentFaces.Reset();
-			for (int32 CornerIndex = 0; CornerIndex < 3; CornerIndex++)
+			//Add this polygon to the group
+			FVertexInfo& OtherVertexInfo = VertexInfoMap.FindOrAdd(PolygonID);
+			//Do not repeat polygons
+			if (!ConsumedPolygon.Contains(OtherVertexInfo.PolygonID))
 			{
-				const FVertexInstanceID& VertexInstanceID = MeshTriangle.GetVertexInstanceID(CornerIndex);
-				int32 ThisCornerIndex = VertexInstanceID.GetValue();
-				DupVerts.Reset();
-				OverlappingCorners.MultiFind(ThisCornerIndex, DupVerts);
-				DupVerts.Add(ThisCornerIndex); // I am a "dup" of myself
-				for (int32 k = 0; k < DupVerts.Num(); k++)
-				{
-					AdjacentFaces.AddUnique(DupVerts[k] / 3);
-				}
+				RecursiveFillPolygonGroup(MeshDescription, VertexInfoMap, CurrentGroup, OtherVertexInfo, ConsumedPolygon);
 			}
+		}
+	}
+}
 
-			// We need to sort these here because the criteria for point equality is
-			// exact, so we must ensure the exact same order for all dups.
-			AdjacentFaces.Sort();
+void FMeshDescriptionHelper::CreateNormals(UMeshDescription* MeshDescription, FMeshDescriptionHelper::ETangentOptions TangentOptions, bool bComputeTangent)
+{
+	//For each vertex compute the normals for every connected edges that are smooth betwween hard edges
+	//         H   A    B
+	//          \  ||  /
+	//       G  -- ** -- C
+	//          // |  \
+	//         F   E    D
+	//
+	// The double ** are the vertex, the double line are hard edges, the single line are soft edge.
+	// A and F are hard, all other edges are soft. The goal is to compute two average normals one from
+	// A to F and a second from F to A. Then we can set the vertex instance normals accordingly.
+	// First normal(A to F) = Normalize(A+B+C+D+E+F)
+	// Second normal(F to A) = Normalize(F+G+H+A)
+	// We found the connected edge using the triangle that share edges
+	
+	TMap<FPolygonID, FVertexInfo> VertexInfoMap;
+	VertexInfoMap.Reserve(20);
+	//Iterate all vertex to compute normals for all vertex instance
+	for (const FVertexID& VertexID : MeshDescription->Vertices().GetElementIDs())
+	{
+		FMeshVertex& Vertex = MeshDescription->GetVertex(VertexID);
 
-			// Process adjacent faces
-			for (int32 AdjacentFaceIndex = 0; AdjacentFaceIndex < AdjacentFaces.Num(); AdjacentFaceIndex++)
+		VertexInfoMap.Reset();
+		
+		//Fill the VertexInfoMap
+		for (const FEdgeID& EdgeID : Vertex.ConnectedEdgeIDs)
+		{
+			FMeshEdge& Edge = MeshDescription->GetEdge(EdgeID);
+			for (const FPolygonID& PolygonID : Edge.ConnectedPolygons)
 			{
-				int32 OtherFaceIndex = AdjacentFaces[AdjacentFaceIndex];
-				for (int32 OurCornerIndex = 0; OurCornerIndex < 3; OurCornerIndex++)
-				{
-					if (bCornerHasNormal[OurCornerIndex])
-						continue;
+				FVertexInfo& VertexInfo = VertexInfoMap.FindOrAdd(PolygonID);
+				int32 EdgeIndex = VertexInfo.EdgeIDs.AddUnique(EdgeID);
 
-					FFanFace NewFanFace;
-					int32 CommonIndexCount = 0;
+				FMeshPolygon& Polygon = MeshDescription->GetPolygon(PolygonID);
 
-					// Check for vertices in common.
-					if (PolygonID.GetValue() == OtherFaceIndex)
-					{
-						CommonIndexCount = 3;
-						NewFanFace.LinkedVertexIndex = OurCornerIndex;
-					}
-					else
+				if (VertexInfo.PolygonID == FPolygonID::Invalid)
+				{
+					VertexInfo.PolygonID = PolygonID;
+					for (FVertexInstanceID& VertexInstanceID : Polygon.PerimeterContour.VertexInstanceIDs)
 					{
-						// Check matching vertices against main vertex .
-						for (int32 OtherCornerIndex = 0; OtherCornerIndex < 3; OtherCornerIndex++)
+						FMeshVertexInstance& VertexInstance = MeshDescription->GetVertexInstance(VertexInstanceID);
+						if (VertexInstance.VertexID == VertexID)
 						{
-							if (CornerPositions[OurCornerIndex].Equals(GetVertexPositionFromVertexInstance(MeshDescription, OtherFaceIndex * 3 + OtherCornerIndex), ComparisonThreshold))
-							{
-								CommonIndexCount++;
-								NewFanFace.LinkedVertexIndex = OtherCornerIndex;
-							}
+							VertexInfo.VertexInstanceID = VertexInstanceID;
+							VertexInfo.UVs = VertexInstance.VertexUVs[0];
+							break;
 						}
 					}
-
-					// Add if connected by at least one point. Smoothing matches are considered later.
-					if (CommonIndexCount > 0)
-					{
-						NewFanFace.FaceIndex = OtherFaceIndex;
-						NewFanFace.bFilled = (OtherFaceIndex == PolygonID.GetValue()); // Starter face for smoothing floodfill.
-						NewFanFace.bBlendTangents = NewFanFace.bFilled;
-						NewFanFace.bBlendNormals = NewFanFace.bFilled;
-						RelevantFacesForCorner[OurCornerIndex].Add(NewFanFace);
-					}
 				}
 			}
+		}
+
+		//Make sure we consume all our vertex instance
+		check(VertexInfoMap.Num() == Vertex.VertexInstanceIDs.Num());
 
-			// Find true relevance of faces for a vertex normal by traversing
-			// smoothing-group-compatible connected triangle fans around common vertices.
-			for (int32 CornerIndex = 0; CornerIndex < 3; CornerIndex++)
+		//Build all group by recursively traverse all polygon connected to the vertex
+		TArray<TSet<FPolygonID>> Groups;
+		TSet<FPolygonID> ConsumedPolygon;
+		for (auto Kvp : VertexInfoMap)
+		{
+			if (ConsumedPolygon.Contains(Kvp.Key))
 			{
-				if (bCornerHasNormal[CornerIndex])
-					continue;
+				continue;
+			}
+
+			int32 CurrentGroupIndex = Groups.AddZeroed();
+			FVertexInfo& VertexInfo = VertexInfoMap.FindOrAdd(Kvp.Key);
+			RecursiveFillPolygonGroup(MeshDescription, VertexInfoMap, Groups[CurrentGroupIndex], VertexInfo, ConsumedPolygon);
+		}
 
-				int32 NewConnections;
-				do
+		//Smooth every connected group
+		ConsumedPolygon.Reset();
+		for(TSet<FPolygonID> Group : Groups)
+		{
+			//Compute tangents data
+			TMap<FVector2D, FVector> GroupTangent;
+			TMap<FVector2D, FVector> GroupBiNormal;
+
+			TArray<FVertexInstanceID> VertexInstanceInGroup;
+			FVector GroupNormal(0.0f);
+			for(FPolygonID& PolygonID : Group)
+			{
+#ifdef ENABLE_NTB_CHECK
+				check(!ConsumedPolygon.Contains(PolygonID));
+#endif
+				ConsumedPolygon.Add(PolygonID);
+				VertexInstanceInGroup.Add(VertexInfoMap[PolygonID].VertexInstanceID);
+				FMeshPolygon& Polygon = MeshDescription->GetPolygon(PolygonID);
+				GroupNormal += Polygon.PolygonNormal;
+				if (bComputeTangent)
 				{
-					NewConnections = 0;
-					for (int32 OtherFaceIdx = 0; OtherFaceIdx < RelevantFacesForCorner[CornerIndex].Num(); OtherFaceIdx++)
-					{
-						FFanFace& OtherFace = RelevantFacesForCorner[CornerIndex][OtherFaceIdx];
-						// The vertex' own face is initially the only face with bFilled == true.
-						if (OtherFace.bFilled)
-						{
-							for (int32 NextFaceIndex = 0; NextFaceIndex < RelevantFacesForCorner[CornerIndex].Num(); NextFaceIndex++)
-							{
-								FFanFace& NextFace = RelevantFacesForCorner[CornerIndex][NextFaceIndex];
-								if (!NextFace.bFilled) // && !NextFace.bBlendTangents)
-								{
-									if ((NextFaceIndex != OtherFaceIdx)
-										&& (SmoothingGroupIndices[NextFace.FaceIndex] & SmoothingGroupIndices[OtherFace.FaceIndex]))
-									{
-										int32 CommonVertices = 0;
-										int32 CommonNormalVertices = 0;
-										for (int32 OtherCornerIndex = 0; OtherCornerIndex < 3; OtherCornerIndex++)
-										{
-											for (int32 NextCornerIndex = 0; NextCornerIndex < 3; NextCornerIndex++)
-											{
-												int32 NextVertexIndex = MeshDescription->GetPolygon(FPolygonID(NextFace.FaceIndex)).Triangles[0].GetVertexInstanceID(NextCornerIndex).GetValue();
-												int32 OtherVertexIndex = MeshDescription->GetPolygon(FPolygonID(OtherFace.FaceIndex)).Triangles[0].GetVertexInstanceID(OtherCornerIndex).GetValue();
-												if (GetVertexPositionFromVertexInstance(MeshDescription, NextVertexIndex).Equals(GetVertexPositionFromVertexInstance(MeshDescription, OtherVertexIndex), ComparisonThreshold))
-												{
-													CommonVertices++;
-													if (bBlendOverlappingNormals
-														|| NextVertexIndex == OtherVertexIndex)
-													{
-														CommonNormalVertices++;
-													}
-												}
-											}
-										}
-										// Flood fill faces with more than one common vertices which must be touching edges.
-										if (CommonVertices > 1)
-										{
-											NextFace.bFilled = true;
-											NextFace.bBlendNormals = (CommonNormalVertices > 1);
-											NewConnections++;
-										}
-									}
-								}
-							}
-						}
-					}
-				} while (NewConnections > 0);
+					FVector2D& UVs = VertexInfoMap[PolygonID].UVs;
+					bool CreateGroup = (!GroupTangent.Contains(UVs));
+					FVector& GroupTangentValue = GroupTangent.FindOrAdd(UVs);
+					FVector& GroupBiNormalValue = GroupBiNormal.FindOrAdd(UVs);
+					GroupTangentValue = CreateGroup ? Polygon.PolygonTangent : GroupTangentValue + Polygon.PolygonTangent;
+					GroupBiNormalValue = CreateGroup ? Polygon.PolygonBinormal : GroupBiNormalValue + Polygon.PolygonBinormal;
+				}
 			}
 
-			// Vertex normal construction.
-			for (int32 CornerIndex = 0; CornerIndex < 3; CornerIndex++)
+			//////////////////////////////////////////////////////////////////////////
+			//Apply the group to the Mesh
+			GroupNormal.Normalize();
+			if (bComputeTangent)
 			{
-				const FVertexInstanceID& VertexInstanceID = MeshTriangle.GetVertexInstanceID(CornerIndex);
-				if (bCornerHasNormal[CornerIndex])
+				for (auto Kvp : GroupTangent)
+				{
+					FVector& GroupTangentValue = GroupTangent.FindOrAdd(Kvp.Key);
+					GroupTangentValue.Normalize();
+				}
+				for (auto Kvp : GroupBiNormal)
 				{
-					CornerNormal[CornerIndex] = OutTangentZ[VertexInstanceID.GetValue()];
+					FVector& GroupBiNormalValue = GroupBiNormal.FindOrAdd(Kvp.Key);
+					GroupBiNormalValue.Normalize();
 				}
-				else
+			}
+			//Apply the average NTB on all Vertex instance
+			for (FVertexInstanceID& VertexInstanceID : VertexInstanceInGroup)
+			{
+				FMeshVertexInstance& VertexInstance = MeshDescription->GetVertexInstance(VertexInstanceID);
+				FVector2D& VertexUV = VertexInstance.VertexUVs[0];
+				
+				//Avoid changing the original group value
+				FVector GroupTangentValue = GroupTangent[VertexUV];
+				FVector GroupBiNormalValue = GroupBiNormal[VertexUV];
+
+				if (VertexInstance.Normal.IsNearlyZero(SMALL_NUMBER))
+				{
+					VertexInstance.Normal = GroupNormal;
+				}
+#ifdef ENABLE_NTB_CHECK
+				check(!VertexInstance.Normal.IsNearlyZero(SMALL_NUMBER));
+#endif
+				if (bComputeTangent)
 				{
-					for (int32 RelevantFaceIdx = 0; RelevantFaceIdx < RelevantFacesForCorner[CornerIndex].Num(); RelevantFaceIdx++)
+					if (!VertexInstance.Tangent.IsNearlyZero(SMALL_NUMBER))
 					{
-						FFanFace const& RelevantFace = RelevantFacesForCorner[CornerIndex][RelevantFaceIdx];
-						if (RelevantFace.bFilled)
-						{
-							int32 OtherFaceIndex = RelevantFace.FaceIndex;
-							if (RelevantFace.bBlendNormals)
-							{
-								CornerNormal[CornerIndex] += TriangleTangentZ[OtherFaceIndex];
-							}
-						}
+						GroupTangentValue = VertexInstance.Tangent;
 					}
-					if (!OutTangentZ[VertexInstanceID.GetValue()].IsZero())
+#ifdef ENABLE_NTB_CHECK
+					check(!GroupTangentValue.IsNearlyZero(SMALL_NUMBER))
+#endif
+					FVector BiNormal(0.0f);
+					if (!VertexInstance.Normal.IsNearlyZero(SMALL_NUMBER) && !VertexInstance.Tangent.IsNearlyZero(SMALL_NUMBER))
 					{
-						CornerNormal[CornerIndex] = OutTangentZ[VertexInstanceID.GetValue()];
+						BiNormal = FVector::CrossProduct(VertexInstance.Normal, VertexInstance.Tangent).GetSafeNormal() * VertexInstance.BinormalSign;
 					}
-				}
-			}
-
-			// Normalization.
-			for (int32 CornerIndex = 0; CornerIndex < 3; CornerIndex++)
-			{
-				CornerNormal[CornerIndex].Normalize();
-			}
+					if (!BiNormal.IsNearlyZero(SMALL_NUMBER))
+					{
+						GroupBiNormalValue = BiNormal;
+					}
+#ifdef ENABLE_NTB_CHECK
+					check(!GroupBiNormalValue.IsNearlyZero(SMALL_NUMBER));
+#endif
+					// Gram-Schmidt orthogonalization
+					GroupBiNormalValue -= GroupTangentValue * (GroupTangentValue | GroupBiNormalValue);
+					GroupBiNormalValue.Normalize();
+					
+					GroupTangentValue -= VertexInstance.Normal * (VertexInstance.Normal | GroupTangentValue);
+					GroupTangentValue.Normalize();
+					
+					GroupBiNormalValue -= VertexInstance.Normal * (VertexInstance.Normal | GroupBiNormalValue);
+					GroupBiNormalValue.Normalize();
+#ifdef ENABLE_NTB_CHECK
+					check(!GroupTangentValue.IsNearlyZero(SMALL_NUMBER));
+					check(!GroupBiNormalValue.IsNearlyZero(SMALL_NUMBER));
+#endif
+					//Set the value
+					VertexInstance.Tangent = GroupTangentValue;
+					//If the BiNormal is zero set the sign to 1.0f
+					VertexInstance.BinormalSign = GetBasisDeterminantSign(GroupTangentValue, GroupBiNormalValue, VertexInstance.Normal);
 
-			// Copy back to the outTangentZ (Normal).
-			for (int32 CornerIndex = 0; CornerIndex < 3; CornerIndex++)
-			{
-				const FVertexInstanceID& VertexInstanceID = MeshTriangle.GetVertexInstanceID(CornerIndex);
-				OutTangentZ[VertexInstanceID.GetValue()] = CornerNormal[CornerIndex];
+				}
 			}
 		}
 	}
+}
 
-	if (OutTangentX.Num() != NumWedges)
-	{
-		OutTangentX.Empty(NumWedges);
-		OutTangentX.AddZeroed(NumWedges);
-	}
-	if (OutTangentY.Num() != NumWedges)
-	{
-		OutTangentY.Empty(NumWedges);
-		OutTangentY.AddZeroed(NumWedges);
-	}
+void FMeshDescriptionHelper::CreateMikktTangents(UMeshDescription* MeshDescription, FMeshDescriptionHelper::ETangentOptions TangentOptions)
+{
+	bool bIgnoreDegenerateTriangles = (TangentOptions & FMeshDescriptionHelper::ETangentOptions::IgnoreDegenerateTriangles) != 0;
+	float ComparisonThreshold = bIgnoreDegenerateTriangles ? THRESH_POINTS_ARE_SAME : 0.0f;
+
+	bool bWedgeTSpace = true;
 
-	//Copy Back the OutTangent to the mesh description
 	for (const FVertexInstanceID& VertexInstanceID : MeshDescription->VertexInstances().GetElementIDs())
 	{
-		int32 VertexInstanceIDValue = VertexInstanceID.GetValue();
-		SetVertexInstanceNTB(MeshDescription, VertexInstanceID, OutTangentZ[VertexInstanceIDValue], OutTangentX[VertexInstanceIDValue], OutTangentY[VertexInstanceIDValue]);
+		FMeshVertexInstance& VertexInstance = MeshDescription->GetVertexInstance(VertexInstanceID);
+		FVector BiNormal(0.0f);
+		if (!VertexInstance.Normal.IsNearlyZero() && !VertexInstance.Tangent.IsNearlyZero())
+		{
+			BiNormal = (FVector::CrossProduct(VertexInstance.Normal, VertexInstance.Tangent).GetSafeNormal() * VertexInstance.BinormalSign).GetSafeNormal();
+		}
+
+		if (VertexInstance.Tangent.IsNearlyZero() || BiNormal.IsNearlyZero())
+		{
+			bWedgeTSpace = false;
+			break;
+		}
 	}
 
+	//Set the tangent space
 	if (!bWedgeTSpace)
 	{
 		// we can use mikktspace to calculate the tangents
@@ -598,10 +575,6 @@ void FMeshDescriptionHelper::ComputeNTB_MikkTSpace(UMeshDescription* MeshDescrip
 		MikkTContext.m_bIgnoreDegenerates = bIgnoreDegenerateTriangles;
 		genTangSpaceDefault(&MikkTContext);
 	}
-
-	check(OutTangentX.Num() == NumWedges);
-	check(OutTangentY.Num() == NumWedges);
-	check(OutTangentZ.Num() == NumWedges);
 }
 
 namespace MeshDescriptionMikktSpaceInterface
@@ -650,7 +623,7 @@ namespace MeshDescriptionMikktSpaceInterface
 			VertexInstance.Tangent[i] = Tangent[i];
 		}
 		//Copy the bi normal sign
-		VertexInstance.BinormalSign = BitangentSign;
+		VertexInstance.BinormalSign = -BitangentSign;
 	}
 
 	void MikkGetTexCoord(const SMikkTSpaceContext* Context, float UV[2], const int FaceIdx, const int VertIdx)
diff --git a/Engine/Source/Runtime/MeshBuilder/Private/MeshDescriptionHelper.h b/Engine/Source/Runtime/MeshBuilder/Private/MeshDescriptionHelper.h
--- a/Engine/Source/Runtime/MeshBuilder/Private/MeshDescriptionHelper.h
+++ b/Engine/Source/Runtime/MeshBuilder/Private/MeshDescriptionHelper.h
@@ -3,6 +3,9 @@
 
 #include "CoreMinimal.h"
 #include "CoreTypes.h"
+#include "Logging/LogMacros.h"
+
+DECLARE_LOG_CATEGORY_EXTERN(LogMeshBuilder, Log, All);
 
 class UObject;
 class UMeshDescription;
@@ -45,12 +48,11 @@ private:
 	//////////////////////////////////////////////////////////////////////////
 	//PRIVATE function declarations
 
-	const FVector& GetVertexPositionFromVertexInstance(UMeshDescription* MeshDescription, int32 VertexInstanceIndex) const;
-	FVector2D GetVertexInstanceUV(UMeshDescription* MeshDescription, int32 VertexInstanceIndex, int32 UVLayer) const;
-	void GetVertexInstanceNTB(UMeshDescription* MeshDescription, const FVertexInstanceID& VertexInstanceID, FVector &OutNormal, FVector &OutTangent, FVector &OutBiNormal) const;
-	void SetVertexInstanceNTB(UMeshDescription* MeshDescription, const FVertexInstanceID& VertexInstanceID, FVector &OutNormal, FVector &OutTangent, FVector &OutBiNormal);
-	void ComputeNTB_MikkTSpace(UMeshDescription* MeshDescription, ETangentOptions TangentOptions);
-	void ComputeTriangleTangents(UMeshDescription* MeshDescription, TArray<FVector>& OutTangentX, TArray<FVector>& OutTangentY, TArray<FVector>& OutTangentZ, float ComparisonThreshold);
+	FORCEINLINE const FVector& GetVertexPositionFromVertexInstance(UMeshDescription* MeshDescription, const FVertexInstanceID& VertexInstanceID) const;
+	FORCEINLINE FVector2D& GetVertexInstanceUV(UMeshDescription* MeshDescription, const FVertexInstanceID& VertexInstanceID, int32 UVLayer) const;
+	void CreateNormals(UMeshDescription* MeshDescription, ETangentOptions TangentOptions, bool bComputeTangent);
+	void CreateMikktTangents(UMeshDescription* MeshDescription, ETangentOptions TangentOptions);
+	void CreatePolygonNTB(UMeshDescription* MeshDescription, float ComparisonThreshold);
 	void FindOverlappingCorners(float ComparisonThreshold);
 
 	//////////////////////////////////////////////////////////////////////////
diff --git a/Engine/Source/Runtime/MeshBuilder/Private/StaticMeshBuilder.cpp b/Engine/Source/Runtime/MeshBuilder/Private/StaticMeshBuilder.cpp
--- a/Engine/Source/Runtime/MeshBuilder/Private/StaticMeshBuilder.cpp
+++ b/Engine/Source/Runtime/MeshBuilder/Private/StaticMeshBuilder.cpp
@@ -11,15 +11,24 @@
 #include "IMeshReductionManagerModule.h"
 #include "MeshBuild.h"
 
-
 //////////////////////////////////////////////////////////////////////////
 //Local functions definition
-int32 GetPolygonGroupTriangles(const class UMeshDescription *MeshDescription, TArray<FMeshTriangle>& OutTriangles, const FPolygonGroupID& PolygonGroupID);
+void GetPolygonGroupTriangles(const class UMeshDescription *MeshDescription, TArray<FPolygonID>& OutPolygons, const FPolygonGroupID& PolygonGroupID);
 bool IsOrphanedVertex(const class UMeshDescription *MeshDescription, const FVertexID VertexID);
 void UpdateBounds(class UStaticMesh* StaticMesh);
 void UpdateCollision(UStaticMesh *StaticMesh);
-void BuildVertexBuffer(const class UMeshDescription* MeshDescription, struct FStaticMeshLODResources& StaticMeshLOD, const struct FMeshBuildSettings& LODBuildSettings, TArray< FStaticMeshBuildVertex >& StaticMeshBuildVertices, const TMultiMap<int32, int32>& OverlappingCorners, float VertexComparisonThreshold, TArray<int32>& RemapVerts);
-void BuildIndexBuffer(class UStaticMesh *StaticMesh, const class UMeshDescription* MeshDescription, TArray< uint32 >& IndexBuffer, struct FStaticMeshLODResources& StaticMeshLOD, TArray<int32>& OutWedgeMap, TArray<TArray<uint32> >& OutPerSectionIndices, TArray<int32>& RemapVerts);
+void BuildVertexBuffer(
+	  class UStaticMesh *StaticMesh
+	, const class UMeshDescription* MeshDescription
+	, struct FStaticMeshLODResources& StaticMeshLOD
+	, const struct FMeshBuildSettings& LODBuildSettings
+	, TArray< uint32 >& IndexBuffer
+	, TArray<int32>& OutWedgeMap
+	, TArray<TArray<uint32> >& OutPerSectionIndices
+	, TArray< FStaticMeshBuildVertex >& StaticMeshBuildVertices
+	, const TMultiMap<int32, int32>& OverlappingCorners
+	, float VertexComparisonThreshold
+	, TArray<int32>& RemapVerts);
 void BuildAllBufferOptimizations(struct FStaticMeshLODResources& StaticMeshLOD, const struct FMeshBuildSettings& LODBuildSettings, TArray< uint32 >& IndexBuffer, bool bNeeds32BitIndices, TArray< FStaticMeshBuildVertex >& StaticMeshBuildVertices);
 //////////////////////////////////////////////////////////////////////////
 
@@ -47,6 +56,7 @@ bool FStaticMeshBuilder::Build(UStaticMesh* StaticMesh)
 		const UMeshDescription* OriginalMeshDescription = StaticMesh->GetOriginalMeshDescription(LodIndex);
 		FMeshDescriptionHelper MeshDescriptionHelper(&LODBuildSettings, OriginalMeshDescription);
 		UMeshDescription* MeshDescription = MeshDescriptionHelper.GetRenderMeshDescription(StaticMesh);
+
 		check(MeshDescription != nullptr);
 		StaticMesh->SetMeshDescription(LodIndex, MeshDescription);
 		
@@ -66,9 +76,7 @@ bool FStaticMeshBuilder::Build(UStaticMesh* StaticMesh)
 
 		StaticMeshLOD.Sections.Empty(PolygonGroups.Num());
 		TArray<int32> RemapVerts; //Because we will remove MeshVertex that are redundant, we need a remap
-		BuildVertexBuffer(MeshDescription, StaticMeshLOD, LODBuildSettings, StaticMeshBuildVertices, MeshDescriptionHelper.GetOverlappingCorners(), VertexComparisonThreshold, RemapVerts);
-
-		//Render data Wedge map is only set for LOD 0???
+								  //Render data Wedge map is only set for LOD 0???
 		TArray<int32> TempWedgeMap;
 		TArray<int32> &WedgeMap = (LodIndex == 0) ? StaticMeshRenderData.WedgeMap : TempWedgeMap;
 
@@ -84,7 +92,8 @@ bool FStaticMeshBuilder::Build(UStaticMesh* StaticMesh)
 			PerSectionIndices.Push(TArray<uint32>());
 		}
 
-		BuildIndexBuffer(StaticMesh, MeshDescription, IndexBuffer, StaticMeshLOD, WedgeMap, PerSectionIndices, RemapVerts);
+		//Build the vertex and index buffer
+		BuildVertexBuffer(StaticMesh, MeshDescription, StaticMeshLOD, LODBuildSettings, IndexBuffer, WedgeMap, PerSectionIndices, StaticMeshBuildVertices, MeshDescriptionHelper.GetOverlappingCorners(), VertexComparisonThreshold, RemapVerts);
 
 		// Figure out which index buffer stride we need
 		bool bNeeds32BitIndices = false;
@@ -93,6 +102,7 @@ bool FStaticMeshBuilder::Build(UStaticMesh* StaticMesh)
 			if (StaticMeshSection.MaxVertexIndex > TNumericLimits<uint16>::Max())
 			{
 				bNeeds32BitIndices = true;
+				break;
 			}
 		}
 		const EIndexBufferStride::Type IndexBufferStride = bNeeds32BitIndices ? EIndexBufferStride::Force32Bit : EIndexBufferStride::Force16Bit;
@@ -107,6 +117,7 @@ bool FStaticMeshBuilder::Build(UStaticMesh* StaticMesh)
 	} //End of LOD for loop
 
 	OnBuildRenderMeshFinish(StaticMesh, true);
+
 	return true;
 }
 
@@ -147,7 +158,7 @@ void FStaticMeshBuilder::OnBuildRenderMeshFinish(UStaticMesh* StaticMesh, const
 	RecreateRenderStateContext.Reset();
 }
 
-int32 GetPolygonGroupTriangles(const UMeshDescription *MeshDescription, TArray<FMeshTriangle>& OutTriangles, const FPolygonGroupID& PolygonGroupID)
+void GetPolygonGroupTriangles(const UMeshDescription *MeshDescription, TArray<FPolygonID>& OutPolygons, const FPolygonGroupID& PolygonGroupID)
 {
 	int32 TriangleCount = 0;
 	for (const FPolygonID& PolygonID : MeshDescription->Polygons().GetElementIDs())
@@ -155,14 +166,9 @@ int32 GetPolygonGroupTriangles(const UMeshDescription *MeshDescription, TArray<F
 		FMeshPolygon MeshPolygon = MeshDescription->GetPolygon(PolygonID);
 		if (MeshPolygon.PolygonGroupID == PolygonGroupID)
 		{
-			for (const FMeshTriangle& MeshTriangle : MeshPolygon.Triangles)
-			{
-				OutTriangles.Add(MeshTriangle);
-				TriangleCount++;
-			}
+			OutPolygons.Add(PolygonID);
 		}
 	}
-	return TriangleCount;
 }
 
 bool IsOrphanedVertex(const UMeshDescription *MeshDescription, const FVertexID VertexID)
@@ -308,10 +314,23 @@ bool AreVerticesEqual(FStaticMeshBuildVertex const& A, FStaticMeshBuildVertex co
 	return true;
 }
 
-void BuildVertexBuffer(const UMeshDescription* MeshDescription, FStaticMeshLODResources& StaticMeshLOD, const FMeshBuildSettings& LODBuildSettings, TArray< FStaticMeshBuildVertex >& StaticMeshBuildVertices, const TMultiMap<int32, int32>& OverlappingCorners, float VertexComparisonThreshold, TArray<int32>& RemapVerts)
+void BuildVertexBuffer(
+	  UStaticMesh *StaticMesh
+	, const UMeshDescription* MeshDescription
+	, FStaticMeshLODResources& StaticMeshLOD
+	, const FMeshBuildSettings& LODBuildSettings
+	, TArray< uint32 >& IndexBuffer
+	, TArray<int32>& OutWedgeMap
+	, TArray<TArray<uint32> >& OutPerSectionIndices
+	, TArray< FStaticMeshBuildVertex >& StaticMeshBuildVertices
+	, const TMultiMap<int32, int32>& OverlappingCorners
+	, float VertexComparisonThreshold
+	, TArray<int32>& RemapVerts)
 {
 	const FVertexArray& Vertices = MeshDescription->Vertices();
 	const FVertexInstanceArray& VertexInstances = MeshDescription->VertexInstances();
+	
+	OutWedgeMap.Reset();
 
 	TArray<int32> RemapVertexInstanceID;
 	// set up vertex buffer elements
@@ -329,19 +348,41 @@ void BuildVertexBuffer(const UMeshDescription* MeshDescription, FStaticMeshLODRe
 	// Set up index buffer
 	for (const FPolygonGroupID& PolygonGroupID : PolygonGroups.GetElementIDs())
 	{
-		TArray<FMeshTriangle> Triangles;
-		int32 SectionTriangleCount = GetPolygonGroupTriangles(MeshDescription, Triangles, PolygonGroupID);
+		TArray<FPolygonID> Polygons;
+		GetPolygonGroupTriangles(MeshDescription, Polygons, PolygonGroupID);
 
 		const FMeshPolygonGroup& PolygonGroup = PolygonGroups[PolygonGroupID];
-		if (Triangles.Num() > 0)
+		
+		TArray<uint32>& SectionIndices = OutPerSectionIndices[PolygonGroupID.GetValue()];
+
+		// Create new rendering section
+		StaticMeshLOD.Sections.Add(FStaticMeshSection());
+		FStaticMeshSection& StaticMeshSection = StaticMeshLOD.Sections.Last();
+
+		StaticMeshSection.FirstIndex = IndexBuffer.Num();
+		StaticMeshSection.NumTriangles = Polygons.Num();
+
+		const int32 MaterialIndex = StaticMesh->GetMaterialIndex(PolygonGroup.ImportedMaterialSlotName);
+		check(MaterialIndex != INDEX_NONE);
+		check(StaticMesh->StaticMaterials[MaterialIndex].MaterialInterface->GetPathName() == PolygonGroup.MaterialAsset.ToString());
+		StaticMeshSection.MaterialIndex = MaterialIndex;
+		StaticMeshSection.bEnableCollision = PolygonGroup.bEnableCollision;
+		StaticMeshSection.bCastShadow = PolygonGroup.bCastShadow;
+
+		for (FPolygonID& PolygonID : Polygons)
 		{
-			for (int32 TriangleIndex = 0; TriangleIndex < Triangles.Num(); ++TriangleIndex)
+			const FMeshPolygon& Polygon = MeshDescription->GetPolygon(PolygonID);
+			int32 ReseveSize = IndexBuffer.Num() + Polygon.Triangles.Num() * 3;
+			IndexBuffer.Reserve(ReseveSize);
+			OutWedgeMap.Reserve(ReseveSize);
+			uint32 MinIndex = TNumericLimits< uint32 >::Max();
+			uint32 MaxIndex = TNumericLimits< uint32 >::Min();
+			for (int32 TriangleIndex = 0; TriangleIndex < Polygon.Triangles.Num(); ++TriangleIndex)
 			{
-				const FMeshTriangle& Triangle = Triangles[TriangleIndex];
+				const FMeshTriangle& Triangle = Polygon.Triangles[TriangleIndex];
 				for (int32 TriVert = 0; TriVert < 3; ++TriVert)
 				{
 					int32 VertexInstanceValue = Triangle.GetVertexInstanceID(TriVert).GetValue();
-					check((TriangleIndex * 3) + TriVert == VertexInstanceValue);
 					const FMeshVertexInstance& VertexInstance = MeshDescription->GetVertexInstance(Triangle.GetVertexInstanceID(TriVert));
 					if (VertexInstance.Color != FColor::White)
 					{
@@ -351,9 +392,10 @@ void BuildVertexBuffer(const UMeshDescription* MeshDescription, FStaticMeshLODRe
 					FStaticMeshBuildVertex StaticMeshVertex;
 
 					StaticMeshVertex.Position = Vertices[VertexInstance.VertexID].VertexPosition * LODBuildSettings.BuildScale3D;
-					StaticMeshVertex.TangentX = VertexInstance.Tangent;
-					StaticMeshVertex.TangentY = FVector::CrossProduct(VertexInstance.Normal, VertexInstance.Tangent).GetSafeNormal() * VertexInstance.BinormalSign;
-					StaticMeshVertex.TangentZ = VertexInstance.Normal;
+					const FMatrix ScaleMatrix = FScaleMatrix(LODBuildSettings.BuildScale3D).Inverse().GetTransposed();
+					StaticMeshVertex.TangentX = ScaleMatrix.TransformVector(VertexInstance.Tangent).GetSafeNormal();
+					StaticMeshVertex.TangentY = ScaleMatrix.TransformVector(FVector::CrossProduct(VertexInstance.Normal, VertexInstance.Tangent).GetSafeNormal() * VertexInstance.BinormalSign).GetSafeNormal();
+					StaticMeshVertex.TangentZ = ScaleMatrix.TransformVector(VertexInstance.Normal).GetSafeNormal();
 					StaticMeshVertex.Color = VertexInstance.Color.ToFColor(true);
 					for (int32 UVIndex = 0; UVIndex < VertexInstance.VertexUVs.Num(); ++UVIndex)
 					{
@@ -385,8 +427,23 @@ void BuildVertexBuffer(const UMeshDescription* MeshDescription, FStaticMeshLODRe
 						Index = StaticMeshBuildVertices.Add(StaticMeshVertex);
 					}
 					RemapVerts[VertexInstanceValue] = Index;
+					const uint32 RenderingVertexIndex = RemapVerts[VertexInstanceValue];
+					IndexBuffer.Add(RenderingVertexIndex);
+					OutWedgeMap.Add(RenderingVertexIndex);
+					SectionIndices.Add(RenderingVertexIndex);
+					MinIndex = FMath::Min(MinIndex, RenderingVertexIndex);
+					MaxIndex = FMath::Max(MaxIndex, RenderingVertexIndex);
 				}
 			}
+			StaticMeshSection.MinVertexIndex = MinIndex;
+			StaticMeshSection.MaxVertexIndex = MaxIndex;
+		}
+
+		if (Polygons.Num() == 0)
+		{
+			// No triangles in this section
+			StaticMeshSection.MinVertexIndex = 0;
+			StaticMeshSection.MaxVertexIndex = 0;
 		}
 	}
 
@@ -404,66 +461,6 @@ void BuildVertexBuffer(const UMeshDescription* MeshDescription, FStaticMeshLODRe
 	}
 }
 
-void BuildIndexBuffer(UStaticMesh *StaticMesh, const UMeshDescription* MeshDescription, TArray< uint32 >& IndexBuffer, FStaticMeshLODResources& StaticMeshLOD, TArray<int32>& OutWedgeMap, TArray<TArray<uint32> >& OutPerSectionIndices, TArray<int32>& RemapVerts)
-{
-	OutWedgeMap.Reset();
-
-	const FPolygonGroupArray& PolygonGroups = MeshDescription->PolygonGroups();
-	// Set up index buffer
-	for (const FPolygonGroupID& PolygonGroupID : PolygonGroups.GetElementIDs())
-	{
-		TArray<uint32>& SectionIndices = OutPerSectionIndices[PolygonGroupID.GetValue()];
-		TArray<FMeshTriangle> Triangles;
-		int32 SectionTriangleCount = GetPolygonGroupTriangles(MeshDescription, Triangles, PolygonGroupID);
-
-		const FMeshPolygonGroup& PolygonGroup = PolygonGroups[PolygonGroupID];
-
-		// Create new rendering section
-		StaticMeshLOD.Sections.Add(FStaticMeshSection());
-		FStaticMeshSection& StaticMeshSection = StaticMeshLOD.Sections.Last();
-
-		StaticMeshSection.FirstIndex = IndexBuffer.Num();
-		StaticMeshSection.NumTriangles = SectionTriangleCount;
-
-		const int32 MaterialIndex = StaticMesh->GetMaterialIndex(PolygonGroup.ImportedMaterialSlotName);
-		check(MaterialIndex != INDEX_NONE);
-		check(StaticMesh->StaticMaterials[MaterialIndex].MaterialInterface->GetPathName() == PolygonGroup.MaterialAsset.ToString());
-		StaticMeshSection.MaterialIndex = MaterialIndex;
-		StaticMeshSection.bEnableCollision = PolygonGroup.bEnableCollision;
-		StaticMeshSection.bCastShadow = PolygonGroup.bCastShadow;
-
-		if (Triangles.Num() > 0)
-		{
-			int32 ReseveSize = IndexBuffer.Num() + Triangles.Num() * 3;
-			IndexBuffer.Reserve(ReseveSize);
-			OutWedgeMap.Reserve(ReseveSize);
-			uint32 MinIndex = TNumericLimits< uint32 >::Max();
-			uint32 MaxIndex = TNumericLimits< uint32 >::Min();
-			for (int32 TriangleIndex = 0; TriangleIndex < Triangles.Num(); ++TriangleIndex)
-			{
-				const FMeshTriangle& Triangle = Triangles[TriangleIndex];
-				for (int32 TriVert = 0; TriVert < 3; ++TriVert)
-				{
-					const uint32 RenderingVertexIndex = RemapVerts[Triangle.GetVertexInstanceID(TriVert).GetValue()];
-					IndexBuffer.Add(RenderingVertexIndex);
-					OutWedgeMap.Add(RenderingVertexIndex);
-					SectionIndices.Add(RenderingVertexIndex);
-					MinIndex = FMath::Min(MinIndex, RenderingVertexIndex);
-					MaxIndex = FMath::Max(MaxIndex, RenderingVertexIndex);
-				}
-			}
-			StaticMeshSection.MinVertexIndex = MinIndex;
-			StaticMeshSection.MaxVertexIndex = MaxIndex;
-		}
-		else
-		{
-			// No triangles in this section
-			StaticMeshSection.MinVertexIndex = 0;
-			StaticMeshSection.MaxVertexIndex = 0;
-		}
-	}
-}
-
 void BuildAllBufferOptimizations(FStaticMeshLODResources& StaticMeshLOD, const FMeshBuildSettings& LODBuildSettings, TArray< uint32 >& IndexBuffer, bool bNeeds32BitIndices, TArray< FStaticMeshBuildVertex >& StaticMeshBuildVertices)
 {
 	const EIndexBufferStride::Type IndexBufferStride = bNeeds32BitIndices ? EIndexBufferStride::Force32Bit : EIndexBufferStride::Force16Bit;
diff --git a/Engine/Source/Runtime/MeshBuilder/Private/ThirdPartyBuildOptimizationHelper.h b/Engine/Source/Runtime/MeshBuilder/Private/ThirdPartyBuildOptimizationHelper.h
new file mode 100644
--- /dev/null
+++ b/Engine/Source/Runtime/MeshBuilder/Private/ThirdPartyBuildOptimizationHelper.h
@@ -0,0 +1,247 @@
+// Copyright 1998-2017 Epic Games, Inc. All Rights Reserved.
+#pragma once
+
+#include "CoreMinimal.h"
+#include "CoreTypes.h"
+//#include "StaticMeshResources.h"
+
+#include "NvTriStrip.h"
+#include "forsythtriangleorderoptimizer.h"
+#include "nvtess.h"
+
+namespace BuildOptimizationPrivate
+{
+	// CVars
+	static TAutoConsoleVariable<int32> CVarTriangleOrderOptimization(
+		TEXT("r.TriangleOrderOptimization"),
+		1,
+		TEXT("Controls the algorithm to use when optimizing the triangle order for the post-transform cache.\n")
+		TEXT("0: Use NVTriStrip (slower)\n")
+		TEXT("1: Use Forsyth algorithm (fastest)(default)")
+		TEXT("2: No triangle order optimization. (least efficient, debugging purposes only)"),
+		ECVF_Default);
+
+	static TAutoConsoleVariable<int32> CVarSupportDepthOnlyIndexBuffers(
+		TEXT("r.SupportDepthOnlyIndexBuffers"),
+		1,
+		TEXT("Enables depth-only index buffers. Saves a little time at the expense of doubling the size of index buffers."),
+		ECVF_ReadOnly | ECVF_RenderThreadSafe);
+
+	static TAutoConsoleVariable<int32> CVarSupportReversedIndexBuffers(
+		TEXT("r.SupportReversedIndexBuffers"),
+		1,
+		TEXT("Enables reversed index buffers. Saves a little time at the expense of doubling the size of index buffers."),
+		ECVF_ReadOnly | ECVF_RenderThreadSafe);
+
+	/*------------------------------------------------------------------------------
+	NVTriStrip for cache optimizing index buffers.
+	------------------------------------------------------------------------------*/
+
+	namespace NvTriStripHelper
+	{
+		/**
+		* Converts 16 bit indices to 32 bit prior to passing them into the real GenerateStrips util method
+		*/
+		void GenerateStrips(
+			const uint8* Indices,
+			bool Is32Bit,
+			const uint32 NumIndices,
+			PrimitiveGroup** PrimGroups,
+			uint32* NumGroups
+		)
+		{
+			if (Is32Bit)
+			{
+				GenerateStrips((uint32*)Indices, NumIndices, PrimGroups, NumGroups);
+			}
+			else
+			{
+				// convert to 32 bit
+				uint32 Idx;
+				TArray<uint32> NewIndices;
+				NewIndices.AddUninitialized(NumIndices);
+				for (Idx = 0; Idx < NumIndices; ++Idx)
+				{
+					NewIndices[Idx] = ((uint16*)Indices)[Idx];
+				}
+				GenerateStrips(NewIndices.GetData(), NumIndices, PrimGroups, NumGroups);
+			}
+
+		}
+
+		/**
+		* Orders a triangle list for better vertex cache coherency.
+		*
+		* *** WARNING: This is safe to call for multiple threads IF AND ONLY IF all
+		* threads call SetListsOnly(true) and SetCacheSize(CACHESIZE_GEFORCE3). If
+		* NvTriStrip is ever used with different settings the library will need
+		* some modifications to be thread-safe. ***
+		*/
+		template<typename IndexDataType, typename Allocator>
+		void CacheOptimizeIndexBuffer(TArray<IndexDataType, Allocator>& Indices)
+		{
+			static_assert(sizeof(IndexDataType) == 2 || sizeof(IndexDataType) == 4, "Indices must be short or int.");
+
+			PrimitiveGroup*	PrimitiveGroups = NULL;
+			uint32			NumPrimitiveGroups = 0;
+			bool Is32Bit = sizeof(IndexDataType) == 4;
+
+			SetListsOnly(true);
+			SetCacheSize(CACHESIZE_GEFORCE3);
+
+			GenerateStrips((uint8*)Indices.GetData(), Is32Bit, Indices.Num(), &PrimitiveGroups, &NumPrimitiveGroups);
+
+			Indices.Empty();
+			Indices.AddUninitialized(PrimitiveGroups->numIndices);
+
+			if (Is32Bit)
+			{
+				FMemory::Memcpy(Indices.GetData(), PrimitiveGroups->indices, Indices.Num() * sizeof(IndexDataType));
+			}
+			else
+			{
+				for (uint32 I = 0; I < PrimitiveGroups->numIndices; ++I)
+				{
+					Indices[I] = (uint16)PrimitiveGroups->indices[I];
+				}
+			}
+
+			delete[] PrimitiveGroups;
+		}
+
+		/**
+		* Provides static mesh render data to the NVIDIA tessellation library.
+		*/
+		class FStaticMeshNvRenderBuffer : public nv::RenderBuffer
+		{
+		public:
+
+			/** Construct from static mesh render buffers. */
+			FStaticMeshNvRenderBuffer(
+				const FPositionVertexBuffer& InPositionVertexBuffer,
+				const FStaticMeshVertexBuffer& InVertexBuffer,
+				const TArray<uint32>& Indices)
+				: PositionVertexBuffer(InPositionVertexBuffer)
+				, VertexBuffer(InVertexBuffer)
+			{
+				check(PositionVertexBuffer.GetNumVertices() == VertexBuffer.GetNumVertices());
+				mIb = new nv::IndexBuffer((void*)Indices.GetData(), nv::IBT_U32, Indices.Num(), false);
+			}
+
+			/** Retrieve the position and first texture coordinate of the specified index. */
+			virtual nv::Vertex getVertex(unsigned int Index) const
+			{
+				nv::Vertex Vertex;
+
+				check(Index < PositionVertexBuffer.GetNumVertices());
+
+				const FVector& Position = PositionVertexBuffer.VertexPosition(Index);
+				Vertex.pos.x = Position.X;
+				Vertex.pos.y = Position.Y;
+				Vertex.pos.z = Position.Z;
+
+				if (VertexBuffer.GetNumTexCoords())
+				{
+					const FVector2D UV = VertexBuffer.GetVertexUV(Index, 0);
+					Vertex.uv.x = UV.X;
+					Vertex.uv.y = UV.Y;
+				}
+				else
+				{
+					Vertex.uv.x = 0.0f;
+					Vertex.uv.y = 0.0f;
+				}
+
+				return Vertex;
+			}
+
+		private:
+
+			/** The position vertex buffer for the static mesh. */
+			const FPositionVertexBuffer& PositionVertexBuffer;
+
+			/** The vertex buffer for the static mesh. */
+			const FStaticMeshVertexBuffer& VertexBuffer;
+
+			/** Copying is forbidden. */
+			FStaticMeshNvRenderBuffer(const FStaticMeshNvRenderBuffer&);
+			FStaticMeshNvRenderBuffer& operator=(const FStaticMeshNvRenderBuffer&);
+		};
+	}
+
+	/*------------------------------------------------------------------------------
+	Forsyth algorithm for cache optimizing index buffers.
+	------------------------------------------------------------------------------*/
+
+	namespace ForsythHelper
+	{
+		/**
+		* Converts 16 bit indices to 32 bit prior to passing them into the real OptimizeFaces util method
+		*/
+		void OptimizeFaces(
+			const uint8* Indices,
+			bool Is32Bit,
+			const uint32 NumIndices,
+			uint32 NumVertices,
+			uint32* OutIndices,
+			uint16 CacheSize
+		)
+		{
+			if (Is32Bit)
+			{
+				Forsyth::OptimizeFaces((uint32*)Indices, NumIndices, NumVertices, OutIndices, CacheSize);
+			}
+			else
+			{
+				// convert to 32 bit
+				uint32 Idx;
+				TArray<uint32> NewIndices;
+				NewIndices.AddUninitialized(NumIndices);
+				for (Idx = 0; Idx < NumIndices; ++Idx)
+				{
+					NewIndices[Idx] = ((uint16*)Indices)[Idx];
+				}
+				Forsyth::OptimizeFaces(NewIndices.GetData(), NumIndices, NumVertices, OutIndices, CacheSize);
+			}
+
+		}
+
+		/**
+		* Orders a triangle list for better vertex cache coherency.
+		*/
+		template<typename IndexDataType, typename Allocator>
+		void CacheOptimizeIndexBuffer(TArray<IndexDataType, Allocator>& Indices)
+		{
+			static_assert(sizeof(IndexDataType) == 2 || sizeof(IndexDataType) == 4, "Indices must be short or int.");
+			bool Is32Bit = sizeof(IndexDataType) == 4;
+
+			// Count the number of vertices
+			uint32 NumVertices = 0;
+			for (int32 Index = 0; Index < Indices.Num(); ++Index)
+			{
+				if (Indices[Index] > NumVertices)
+				{
+					NumVertices = Indices[Index];
+				}
+			}
+			NumVertices += 1;
+
+			TArray<uint32> OptimizedIndices;
+			OptimizedIndices.AddUninitialized(Indices.Num());
+			uint16 CacheSize = 32;
+			OptimizeFaces((uint8*)Indices.GetData(), Is32Bit, Indices.Num(), NumVertices, OptimizedIndices.GetData(), CacheSize);
+
+			if (Is32Bit)
+			{
+				FMemory::Memcpy(Indices.GetData(), OptimizedIndices.GetData(), Indices.Num() * sizeof(IndexDataType));
+			}
+			else
+			{
+				for (int32 I = 0; I < OptimizedIndices.Num(); ++I)
+				{
+					Indices[I] = (uint16)OptimizedIndices[I];
+				}
+			}
+		}
+	}
+}
diff --git a/Engine/Source/Runtime/MeshDescription/Public/MeshDescription.h b/Engine/Source/Runtime/MeshDescription/Public/MeshDescription.h
--- a/Engine/Source/Runtime/MeshDescription/Public/MeshDescription.h
+++ b/Engine/Source/Runtime/MeshDescription/Public/MeshDescription.h
@@ -666,7 +666,7 @@ public:
 	}
 
 	/** Deletes a vertex instance from a mesh */
-	void DeleteVertexInstance( const FVertexInstanceID VertexInstanceID, TArray<FVertexID>* InOutOrphanedVerticesPtr = nullptr )
+	void DeleteVertexInstance( const FVertexInstanceID VertexInstanceID, TSet<FVertexID>* InOutOrphanedVerticesPtr = nullptr )
 	{
 		const FVertexID VertexID = VertexInstanceArray[ VertexInstanceID ].VertexID;
 		verify( VertexArray[ VertexID ].VertexInstanceIDs.Remove( VertexInstanceID ) == 1 );
@@ -705,7 +705,7 @@ public:
 	}
 
 	/** Deletes an edge from a mesh */
-	void DeleteEdge( const FEdgeID EdgeID, TArray<FVertexID>* InOutOrphanedVerticesPtr = nullptr )
+	void DeleteEdge( const FEdgeID EdgeID, TSet<FVertexID>* InOutOrphanedVerticesPtr = nullptr )
 	{
 		FMeshEdge& Edge = EdgeArray[ EdgeID ];
 		for( const FVertexID EdgeVertexID : Edge.VertexIDs )
@@ -774,7 +774,7 @@ public:
 	}
 
 private:
-	void DeletePolygonContour_Internal( const FPolygonID PolygonID, const TArray<FVertexInstanceID>& VertexInstanceIDs, TArray<FEdgeID>* InOutOrphanedEdgesPtr, TArray<FVertexInstanceID>* InOutOrphanedVertexInstancesPtr )
+	void DeletePolygonContour_Internal( const FPolygonID PolygonID, const TArray<FVertexInstanceID>& VertexInstanceIDs, TSet<FEdgeID>* InOutOrphanedEdgesPtr, TSet<FVertexInstanceID>* InOutOrphanedVertexInstancesPtr )
 	{
 		FVertexInstanceID LastVertexInstanceID = VertexInstanceIDs.Last();
 		for( const FVertexInstanceID VertexInstanceID : VertexInstanceIDs )
@@ -804,7 +804,7 @@ private:
 
 public:
 	/** Deletes a polygon from the mesh */
-	void DeletePolygon( const FPolygonID PolygonID, TArray<FEdgeID>* InOutOrphanedEdgesPtr = nullptr, TArray<FVertexInstanceID>* InOutOrphanedVertexInstancesPtr = nullptr, TArray<FPolygonGroupID>* InOutOrphanedPolygonGroupsPtr = nullptr )
+	void DeletePolygon( const FPolygonID PolygonID, TSet<FEdgeID>* InOutOrphanedEdgesPtr = nullptr, TSet<FVertexInstanceID>* InOutOrphanedVertexInstancesPtr = nullptr, TSet<FPolygonGroupID>* InOutOrphanedPolygonGroupsPtr = nullptr )
 	{
 		FMeshPolygon& Polygon = PolygonArray[ PolygonID ];
 		DeletePolygonContour_Internal( PolygonID, Polygon.PerimeterContour.VertexInstanceIDs, InOutOrphanedEdgesPtr, InOutOrphanedVertexInstancesPtr );
