Add some performance log
Make some small optimisation by using the meshdescription array directly

diff --git a/Engine/Source/Runtime/MeshBuilder/Private/BuildStatisticManager.h b/Engine/Source/Runtime/MeshBuilder/Private/BuildStatisticManager.h
new file mode 100644
--- /dev/null
+++ b/Engine/Source/Runtime/MeshBuilder/Private/BuildStatisticManager.h
@@ -0,0 +1,28 @@
+// Copyright 1998-2017 Epic Games, Inc. All Rights Reserved.
+#pragma once
+
+#include "CoreMinimal.h"
+#include "CoreTypes.h"
+
+DECLARE_LOG_CATEGORY_EXTERN(LogBuildStatistic, Log, All);
+
+namespace BuildStatisticManager
+{
+	class FBuildStatisticScope
+	{
+	public:
+		FBuildStatisticScope(const FString& InTimerDescription)
+		{
+			TimerDescription = InTimerDescription;
+			StartTime = FPlatformTime::Cycles();
+		}
+		~FBuildStatisticScope()
+		{
+			float ScopeTime = FPlatformTime::ToSeconds(FPlatformTime::Cycles() - StartTime);
+			UE_LOG(LogBuildStatistic, Log, TEXT("%s: %f seconds"), *TimerDescription, ScopeTime);
+		}
+	private:
+		int32 StartTime;
+		FString TimerDescription;
+	};
+}
\ No newline at end of file
diff --git a/Engine/Source/Runtime/MeshBuilder/Private/MeshDescriptionHelper.cpp b/Engine/Source/Runtime/MeshBuilder/Private/MeshDescriptionHelper.cpp
--- a/Engine/Source/Runtime/MeshBuilder/Private/MeshDescriptionHelper.cpp
+++ b/Engine/Source/Runtime/MeshBuilder/Private/MeshDescriptionHelper.cpp
@@ -14,10 +14,8 @@
 #include "RenderUtils.h"
 #include "LayoutUV.h"
 #include "IMeshReductionInterfaces.h"
-
 #include "IMeshReductionManagerModule.h"
-
-DEFINE_LOG_CATEGORY(LogMeshBuilder);
+#include "BuildStatisticManager.h"
 
 //Enable all check
 //#define ENABLE_NTB_CHECK
@@ -41,6 +39,7 @@ FMeshDescriptionHelper::FMeshDescriptionHelper(FMeshBuildSettings* InBuildSettin
 
 UMeshDescription* FMeshDescriptionHelper::GetRenderMeshDescription(UObject* Owner)
 {
+	BuildStatisticManager::FBuildStatisticScope StatScope(TEXT("GetRenderMeshDescription took"));
 	UStaticMesh* StaticMesh = Cast<UStaticMesh>(Owner);
 	check(StaticMesh);
 	//Use the build settings to create the RenderMeshDescription
@@ -55,11 +54,13 @@ UMeshDescription* FMeshDescriptionHelper::GetRenderMeshDescription(UObject* Owne
 	CopyMeshDescription(const_cast<UMeshDescription*>(OriginalMeshDescription), RenderMeshDescription);
 	float ComparisonThreshold = BuildSettings->bRemoveDegenerates ? THRESH_POINTS_ARE_SAME : 0.0f;
 
+	FVertexInstanceArray& VertexInstanceArray = RenderMeshDescription->VertexInstances();
+
 	if (BuildSettings->bRecomputeTangents || BuildSettings->bRecomputeNormals)
 	{
-		for (const FVertexInstanceID& VertexInstanceID : RenderMeshDescription->VertexInstances().GetElementIDs())
+		for (const FVertexInstanceID& VertexInstanceID : VertexInstanceArray.GetElementIDs())
 		{
-			FMeshVertexInstance& VertexInstance = RenderMeshDescription->GetVertexInstance(VertexInstanceID);
+			FMeshVertexInstance& VertexInstance = VertexInstanceArray[VertexInstanceID];
 			// Dump normals and tangents if we are recomputing them.
 			if (BuildSettings->bRecomputeTangents)
 			{
@@ -108,7 +109,7 @@ UMeshDescription* FMeshDescriptionHelper::GetRenderMeshDescription(UObject* Owne
 		}
 	}
 
-	if (BuildSettings->bGenerateLightmapUVs && RenderMeshDescription->VertexInstances().Num() > 0)
+	if (BuildSettings->bGenerateLightmapUVs && VertexInstanceArray.Num() > 0)
 	{
 		if (!RenderMeshDescription->GetVertexInstance(FVertexInstanceID(0)).VertexUVs.IsValidIndex(BuildSettings->SrcLightmapIndex))
 		{
@@ -168,28 +169,22 @@ bool FMeshDescriptionHelper::IsValidOriginalMeshDescription()
 
 void FMeshDescriptionHelper::FindOverlappingCorners(TMultiMap<int32, int32>& OverlappingCorners, const UMeshDescription* MeshDescription, float ComparisonThreshold)
 {
+	BuildStatisticManager::FBuildStatisticScope StatScope(TEXT("FindOverlappingCorners took"));
 	//Empty the old data
 	OverlappingCorners.Reset();
-	const int32 NumWedges = MeshDescription->VertexInstances().Num();
+	
+	const FVertexInstanceArray& VertexInstanceArray = MeshDescription->VertexInstances();
+	const FVertexArray& VertexArray = MeshDescription->Vertices();
+
+	const int32 NumWedges = VertexInstanceArray.Num();
 
 	// Create a list of vertex Z/index pairs
 	TArray<FIndexAndZ> VertIndexAndZ;
 	VertIndexAndZ.Reserve(NumWedges);
-	const FVertexInstanceArray& VertexInstances = MeshDescription->VertexInstances();
-
-	for (const FPolygonID& PolygonID : MeshDescription->Polygons().GetElementIDs())
+	for (const FVertexInstanceID VertexInstanceID : VertexInstanceArray.GetElementIDs())
 	{
-		const TArray<FMeshTriangle>& MeshTriangles = MeshDescription->GetPolygon(PolygonID).Triangles;
-		for (const FMeshTriangle& MeshTriangle : MeshTriangles)
-		{
-			for (int32 CornerIndex = 0; CornerIndex < 3; ++CornerIndex)
-			{
-				const FVertexInstanceID& VertexInstanceID = MeshTriangle.GetVertexInstanceID(CornerIndex);
-				const FMeshVertexInstance& VertexInstance = MeshDescription->GetVertexInstance(VertexInstanceID);
-
-				new(VertIndexAndZ)FIndexAndZ(VertexInstanceID.GetValue(), MeshDescription->GetVertex(VertexInstance.VertexID).VertexPosition);
-			}
-		}
+		const FMeshVertexInstance& VertexInstance = VertexInstanceArray[VertexInstanceID];
+		new(VertIndexAndZ)FIndexAndZ(VertexInstanceID.GetValue(), VertexArray[VertexInstance.VertexID].VertexPosition);
 	}
 
 	// Sort the vertices by z value
@@ -204,8 +199,8 @@ void FMeshDescriptionHelper::FindOverlappingCorners(TMultiMap<int32, int32>& Ove
 			if (FMath::Abs(VertIndexAndZ[j].Z - VertIndexAndZ[i].Z) > ComparisonThreshold)
 				break; // can't be any more dups
 
-			const FVector& PositionA = MeshDescription->GetVertex(MeshDescription->GetVertexInstance(FVertexInstanceID(VertIndexAndZ[i].Index)).VertexID).VertexPosition;
-			const FVector& PositionB = MeshDescription->GetVertex(MeshDescription->GetVertexInstance(FVertexInstanceID(VertIndexAndZ[j].Index)).VertexID).VertexPosition;
+			const FVector& PositionA = *(VertIndexAndZ[i].OriginalVector);
+			const FVector& PositionB = *(VertIndexAndZ[j].OriginalVector);
 
 			if (PositionA.Equals(PositionB, ComparisonThreshold))
 			{
@@ -221,26 +216,22 @@ void FMeshDescriptionHelper::FindOverlappingCorners(const UMeshDescription* Mesh
 	FindOverlappingCorners(OverlappingCorners, MeshDescription, ComparisonThreshold);
 }
 
-const FVector& FMeshDescriptionHelper::GetVertexPositionFromVertexInstance(UMeshDescription* MeshDescription, const FVertexInstanceID& VertexInstanceID) const
-{
-	return MeshDescription->GetVertex(MeshDescription->GetVertexInstance(VertexInstanceID).VertexID).VertexPosition;
-}
-
-FVector2D &FMeshDescriptionHelper::GetVertexInstanceUV(UMeshDescription* MeshDescription, const FVertexInstanceID& VertexInstanceID, int32 UVLayer) const
-{
-	return MeshDescription->GetVertexInstance(VertexInstanceID).VertexUVs[UVLayer];
-}
-
 void FMeshDescriptionHelper::CreatePolygonNTB(UMeshDescription* MeshDescription, float ComparisonThreshold)
 {
+	BuildStatisticManager::FBuildStatisticScope StatScope(TEXT("CreatePolygonNTB took"));
 	TArray<FPolygonID> DegeneratePolygons;
+
+	FVertexInstanceArray& VertexInstanceArray = MeshDescription->VertexInstances();
+	FVertexArray& VertexArray = MeshDescription->Vertices();
+	FPolygonArray& PolygonArray = MeshDescription->Polygons();
+
 	for (const FPolygonID& PolygonID : MeshDescription->Polygons().GetElementIDs())
 	{
 		FVector TangentX(0.0f);
 		FVector TangentY(0.0f);
 		FVector TangentZ(0.0f);
 
-		FMeshPolygon& Polygon = MeshDescription->GetPolygon(PolygonID);
+		FMeshPolygon& Polygon = PolygonArray[PolygonID];
 		const TArray<FMeshTriangle>& MeshTriangles = Polygon.Triangles;
 #ifdef ENABLE_NTB_CHECK
 		//Assume triangle are build
@@ -252,9 +243,15 @@ void FMeshDescriptionHelper::CreatePolygonNTB(UMeshDescription* MeshDescription,
 		int32 UVIndex = 0;
 
 		FVector P[3];
+		FVector2D UVs[3];
+		FMeshVertexInstance* MeshVertexInstance[3];
+		FMeshVertex* MeshVertex[3];
 		for (int32 i = 0; i < 3; ++i)
 		{
-			P[i] = GetVertexPositionFromVertexInstance(MeshDescription, MeshTriangle.GetVertexInstanceID(i));
+			MeshVertexInstance[i] = &(VertexInstanceArray[MeshTriangle.GetVertexInstanceID(i)]);
+			UVs[i] = MeshVertexInstance[i]->VertexUVs[0];
+			MeshVertex[i] = &(VertexArray[MeshVertexInstance[i]->VertexID]);
+			P[i] = MeshVertex[i]->VertexPosition;
 		}
 
 		const FVector Normal = ((P[1] - P[2]) ^ (P[0] - P[2])).GetSafeNormal(ComparisonThreshold);
@@ -268,14 +265,10 @@ void FMeshDescriptionHelper::CreatePolygonNTB(UMeshDescription* MeshDescription,
 				FPlane(0, 0, 0, 1)
 			);
 
-			const FVector2D T1 = GetVertexInstanceUV(MeshDescription, MeshTriangle.GetVertexInstanceID(0), 0);
-			const FVector2D T2 = GetVertexInstanceUV(MeshDescription, MeshTriangle.GetVertexInstanceID(1), 0);
-			const FVector2D T3 = GetVertexInstanceUV(MeshDescription, MeshTriangle.GetVertexInstanceID(2), 0);
-
 			FMatrix ParameterToTexture(
-				FPlane(T2.X - T1.X, T2.Y - T1.Y, 0, 0),
-				FPlane(T3.X - T1.X, T3.Y - T1.Y, 0, 0),
-				FPlane(T1.X, T1.Y, 1, 0),
+				FPlane(UVs[1].X - UVs[0].X, UVs[1].Y - UVs[0].Y, 0, 0),
+				FPlane(UVs[2].X - UVs[0].X, UVs[2].Y - UVs[0].Y, 0, 0),
+				FPlane(UVs[0].X, UVs[0].Y, 1, 0),
 				FPlane(0, 0, 0, 1)
 			);
 
@@ -359,9 +352,12 @@ void RecursiveFillPolygonGroup(
 {
 	CurrentGroup.Add(CurrentVertexInfo.PolygonID);
 	ConsumedPolygon.Add(CurrentVertexInfo.PolygonID);
+
+	FEdgeArray& EdgeArray = MeshDescription->Edges();
+
 	for (FEdgeID &EdgeID : CurrentVertexInfo.EdgeIDs)
 	{
-		FMeshEdge& Edge = MeshDescription->GetEdge(EdgeID);
+		FMeshEdge& Edge = EdgeArray[EdgeID];
 		if (Edge.bIsHardEdge)
 		{
 			//End of the group
@@ -386,6 +382,7 @@ void RecursiveFillPolygonGroup(
 
 void FMeshDescriptionHelper::CreateNormals(UMeshDescription* MeshDescription, FMeshDescriptionHelper::ETangentOptions TangentOptions, bool bComputeTangent)
 {
+	BuildStatisticManager::FBuildStatisticScope StatScope(TEXT("CreateNormals took"));
 	//For each vertex compute the normals for every connected edges that are smooth betwween hard edges
 	//         H   A    B
 	//          \  ||  /
@@ -400,32 +397,37 @@ void FMeshDescriptionHelper::CreateNormals(UMeshDescription* MeshDescription, FM
 	// Second normal(F to A) = Normalize(F+G+H+A)
 	// We found the connected edge using the triangle that share edges
 	
+	FVertexInstanceArray& VertexInstanceArray = MeshDescription->VertexInstances();
+	FVertexArray& VertexArray = MeshDescription->Vertices();
+	FPolygonArray& PolygonArray = MeshDescription->Polygons();
+	FEdgeArray& EdgeArray = MeshDescription->Edges();
+
 	TMap<FPolygonID, FVertexInfo> VertexInfoMap;
 	VertexInfoMap.Reserve(20);
 	//Iterate all vertex to compute normals for all vertex instance
-	for (const FVertexID& VertexID : MeshDescription->Vertices().GetElementIDs())
+	for (const FVertexID& VertexID : VertexArray.GetElementIDs())
 	{
-		FMeshVertex& Vertex = MeshDescription->GetVertex(VertexID);
+		FMeshVertex& Vertex = VertexArray[VertexID];
 
 		VertexInfoMap.Reset();
 		
 		//Fill the VertexInfoMap
 		for (const FEdgeID& EdgeID : Vertex.ConnectedEdgeIDs)
 		{
-			FMeshEdge& Edge = MeshDescription->GetEdge(EdgeID);
+			FMeshEdge& Edge = EdgeArray[EdgeID];
 			for (const FPolygonID& PolygonID : Edge.ConnectedPolygons)
 			{
 				FVertexInfo& VertexInfo = VertexInfoMap.FindOrAdd(PolygonID);
 				int32 EdgeIndex = VertexInfo.EdgeIDs.AddUnique(EdgeID);
 
-				FMeshPolygon& Polygon = MeshDescription->GetPolygon(PolygonID);
-
 				if (VertexInfo.PolygonID == FPolygonID::Invalid)
 				{
+					FMeshPolygon& Polygon = PolygonArray[PolygonID];
+
 					VertexInfo.PolygonID = PolygonID;
 					for (FVertexInstanceID& VertexInstanceID : Polygon.PerimeterContour.VertexInstanceIDs)
 					{
-						FMeshVertexInstance& VertexInstance = MeshDescription->GetVertexInstance(VertexInstanceID);
+						FMeshVertexInstance& VertexInstance = VertexInstanceArray[VertexInstanceID];
 						if (VertexInstance.VertexID == VertexID)
 						{
 							VertexInfo.VertexInstanceID = VertexInstanceID;
@@ -472,7 +474,7 @@ void FMeshDescriptionHelper::CreateNormals(UMeshDescription* MeshDescription, FM
 #endif
 				ConsumedPolygon.Add(PolygonID);
 				VertexInstanceInGroup.Add(VertexInfoMap[PolygonID].VertexInstanceID);
-				FMeshPolygon& Polygon = MeshDescription->GetPolygon(PolygonID);
+				FMeshPolygon& Polygon = PolygonArray[PolygonID];
 				GroupNormal += Polygon.PolygonNormal;
 				if (bComputeTangent)
 				{
@@ -504,7 +506,7 @@ void FMeshDescriptionHelper::CreateNormals(UMeshDescription* MeshDescription, FM
 			//Apply the average NTB on all Vertex instance
 			for (FVertexInstanceID& VertexInstanceID : VertexInstanceInGroup)
 			{
-				FMeshVertexInstance& VertexInstance = MeshDescription->GetVertexInstance(VertexInstanceID);
+				FMeshVertexInstance& VertexInstance = VertexInstanceArray[VertexInstanceID];
 				FVector2D& VertexUV = VertexInstance.VertexUVs[0];
 			
 				if (VertexInstance.Normal.IsNearlyZero(SMALL_NUMBER))
@@ -565,6 +567,7 @@ void FMeshDescriptionHelper::CreateNormals(UMeshDescription* MeshDescription, FM
 
 void FMeshDescriptionHelper::CreateMikktTangents(UMeshDescription* MeshDescription, FMeshDescriptionHelper::ETangentOptions TangentOptions)
 {
+	BuildStatisticManager::FBuildStatisticScope StatScope(TEXT("CreateMikktTangents took"));
 	bool bIgnoreDegenerateTriangles = (TangentOptions & FMeshDescriptionHelper::ETangentOptions::IgnoreDegenerateTriangles) != 0;
 	float ComparisonThreshold = bIgnoreDegenerateTriangles ? THRESH_POINTS_ARE_SAME : 0.0f;
 
@@ -602,7 +605,7 @@ namespace MeshDescriptionMikktSpaceInterface
 	void MikkGetPosition(const SMikkTSpaceContext* Context, float Position[3], const int FaceIdx, const int VertIdx)
 	{
 		UMeshDescription *MeshDescription = (UMeshDescription*)(Context->m_pUserData);
-		FVector VertexPosition = MeshDescription->GetVertex(MeshDescription->GetVertexInstance(FVertexInstanceID(FaceIdx * 3 + VertIdx)).VertexID).VertexPosition;
+		const FVector &VertexPosition = MeshDescription->Vertices()[MeshDescription->VertexInstances()[FVertexInstanceID(FaceIdx * 3 + VertIdx)].VertexID].VertexPosition;
 		Position[0] = VertexPosition.X;
 		Position[1] = VertexPosition.Y;
 		Position[2] = VertexPosition.Z;
@@ -611,8 +614,7 @@ namespace MeshDescriptionMikktSpaceInterface
 	void MikkGetNormal(const SMikkTSpaceContext* Context, float Normal[3], const int FaceIdx, const int VertIdx)
 	{
 		UMeshDescription *MeshDescription = (UMeshDescription*)(Context->m_pUserData);
-		FVertexInstanceID VertexInstanceID(FaceIdx * 3 + VertIdx);
-		const FMeshVertexInstance& VertexInstance = MeshDescription->GetVertexInstance(VertexInstanceID);
+		const FMeshVertexInstance& VertexInstance = MeshDescription->VertexInstances()[FVertexInstanceID(FaceIdx * 3 + VertIdx)];
 		const FVector &VertexNormal = VertexInstance.Normal;
 		for (int32 i = 0; i < 3; ++i)
 		{
@@ -623,8 +625,7 @@ namespace MeshDescriptionMikktSpaceInterface
 	void MikkSetTSpaceBasic(const SMikkTSpaceContext* Context, const float Tangent[3], const float BitangentSign, const int FaceIdx, const int VertIdx)
 	{
 		UMeshDescription *MeshDescription = (UMeshDescription*)(Context->m_pUserData);
-		FVertexInstanceID VertexInstanceID(FaceIdx * 3 + VertIdx);
-		FMeshVertexInstance& VertexInstance = MeshDescription->GetVertexInstance(VertexInstanceID);
+		FMeshVertexInstance& VertexInstance = MeshDescription->VertexInstances()[FVertexInstanceID(FaceIdx * 3 + VertIdx)];
 		//Copy the tangents
 		for (int32 i = 0; i < 3; ++i)
 		{
@@ -637,7 +638,7 @@ namespace MeshDescriptionMikktSpaceInterface
 	void MikkGetTexCoord(const SMikkTSpaceContext* Context, float UV[2], const int FaceIdx, const int VertIdx)
 	{
 		UMeshDescription *MeshDescription = (UMeshDescription*)(Context->m_pUserData);
-		const FVector2D &TexCoord = MeshDescription->GetVertexInstance(FVertexInstanceID(FaceIdx * 3 + VertIdx)).VertexUVs[0];
+		const FVector2D &TexCoord = MeshDescription->VertexInstances()[FVertexInstanceID(FaceIdx * 3 + VertIdx)].VertexUVs[0];
 		UV[0] = TexCoord.X;
 		UV[1] = TexCoord.Y;
 	}
diff --git a/Engine/Source/Runtime/MeshBuilder/Private/MeshDescriptionHelper.h b/Engine/Source/Runtime/MeshBuilder/Private/MeshDescriptionHelper.h
--- a/Engine/Source/Runtime/MeshBuilder/Private/MeshDescriptionHelper.h
+++ b/Engine/Source/Runtime/MeshBuilder/Private/MeshDescriptionHelper.h
@@ -80,15 +80,17 @@ private:
 	{
 		float Z;
 		int32 Index;
+		const FVector *OriginalVector;
 
 		/** Default constructor. */
 		FIndexAndZ() {}
 
 		/** Initialization constructor. */
-		FIndexAndZ(int32 InIndex, FVector V)
+		FIndexAndZ(int32 InIndex, const FVector& V)
 		{
 			Z = 0.30f * V.X + 0.33f * V.Y + 0.37f * V.Z;
 			Index = InIndex;
+			OriginalVector = &V;
 		}
 	};
 
diff --git a/Engine/Source/Runtime/MeshBuilder/Private/StaticMeshBuilder.cpp b/Engine/Source/Runtime/MeshBuilder/Private/StaticMeshBuilder.cpp
--- a/Engine/Source/Runtime/MeshBuilder/Private/StaticMeshBuilder.cpp
+++ b/Engine/Source/Runtime/MeshBuilder/Private/StaticMeshBuilder.cpp
@@ -10,19 +10,21 @@
 #include "Components.h"
 #include "IMeshReductionManagerModule.h"
 #include "MeshBuild.h"
+#include "BuildStatisticManager.h"
+
+DEFINE_LOG_CATEGORY(LogBuildStatistic);
 
 //////////////////////////////////////////////////////////////////////////
 //Local functions definition
-void GetPolygonGroupTriangles(const class UMeshDescription *MeshDescription, TArray<FPolygonID>& OutPolygons, const FPolygonGroupID& PolygonGroupID);
 bool IsOrphanedVertex(const class UMeshDescription *MeshDescription, const FVertexID VertexID);
 void UpdateBounds(class UStaticMesh* StaticMesh);
 void UpdateCollision(UStaticMesh *StaticMesh);
 void BuildVertexBuffer(
-	  class UStaticMesh *StaticMesh
+	  UStaticMesh *StaticMesh
 	, int32 LodIndex
-	, const class UMeshDescription* MeshDescription
-	, struct FStaticMeshLODResources& StaticMeshLOD
-	, const struct FMeshBuildSettings& LODBuildSettings
+	, UMeshDescription* MeshDescription
+	, FStaticMeshLODResources& StaticMeshLOD
+	, const FMeshBuildSettings& LODBuildSettings
 	, TArray< uint32 >& IndexBuffer
 	, TArray<int32>& OutWedgeMap
 	, TArray<TArray<uint32> >& OutPerSectionIndices
@@ -40,6 +42,7 @@ FStaticMeshBuilder::FStaticMeshBuilder()
 
 bool FStaticMeshBuilder::Build(UStaticMesh* StaticMesh, const FStaticMeshLODGroup& LODGroup)
 {
+	BuildStatisticManager::FBuildStatisticScope StatScope(TEXT("Building StaticMesh took"));
 	if (StaticMesh->GetOriginalMeshDescription(0) == nullptr)
 	{
 		//TODO: Warn the user that there is no mesh description data
@@ -184,23 +187,10 @@ void FStaticMeshBuilder::OnBuildRenderMeshFinish(UStaticMesh* StaticMesh, const
 	RecreateRenderStateContext.Reset();
 }
 
-void GetPolygonGroupTriangles(const UMeshDescription *MeshDescription, TArray<FPolygonID>& OutPolygons, const FPolygonGroupID& PolygonGroupID)
-{
-	int32 TriangleCount = 0;
-	for (const FPolygonID& PolygonID : MeshDescription->Polygons().GetElementIDs())
-	{
-		FMeshPolygon MeshPolygon = MeshDescription->GetPolygon(PolygonID);
-		if (MeshPolygon.PolygonGroupID == PolygonGroupID)
-		{
-			OutPolygons.Add(PolygonID);
-		}
-	}
-}
-
 bool IsOrphanedVertex(const UMeshDescription *MeshDescription, const FVertexID VertexID)
 {
 	const FVertexInstanceArray& VertexInstances = MeshDescription->VertexInstances();
-	for (const FVertexInstanceID VertexInstanceID : MeshDescription->GetVertex(VertexID).VertexInstanceIDs)
+	for (const FVertexInstanceID VertexInstanceID : MeshDescription->Vertices()[VertexID].VertexInstanceIDs)
 	{
 		if (VertexInstances[VertexInstanceID].ConnectedPolygons.Num() > 0)
 		{
@@ -343,7 +333,7 @@ bool AreVerticesEqual(FStaticMeshBuildVertex const& A, FStaticMeshBuildVertex co
 void BuildVertexBuffer(
 	  UStaticMesh *StaticMesh
 	, int32 LodIndex
-	, const UMeshDescription* MeshDescription
+	, UMeshDescription* MeshDescription
 	, FStaticMeshLODResources& StaticMeshLOD
 	, const FMeshBuildSettings& LODBuildSettings
 	, TArray< uint32 >& IndexBuffer
@@ -354,8 +344,11 @@ void BuildVertexBuffer(
 	, float VertexComparisonThreshold
 	, TArray<int32>& RemapVerts)
 {
-	const FVertexArray& Vertices = MeshDescription->Vertices();
-	const FVertexInstanceArray& VertexInstances = MeshDescription->VertexInstances();
+	BuildStatisticManager::FBuildStatisticScope StatScope(TEXT("BuildVertexBuffer took"));
+	FVertexArray& Vertices = MeshDescription->Vertices();
+	FVertexInstanceArray& VertexInstances = MeshDescription->VertexInstances();
+	FPolygonGroupArray& PolygonGroupArray = MeshDescription->PolygonGroups();
+	FPolygonArray& PolygonArray = MeshDescription->Polygons();
 	
 	OutWedgeMap.Reset();
 
@@ -371,16 +364,13 @@ void BuildVertexBuffer(
 	}
 	TArray<int32> DupVerts;
 
-	uint32 NumTextureCoord = (uint32)(MeshDescription->VertexInstances().Num() > 0 ? MeshDescription->GetVertexInstance(FVertexInstanceID(0)).VertexUVs.Num() : 1);
+	uint32 NumTextureCoord = (uint32)(VertexInstances.Num() > 0 ? VertexInstances[FVertexInstanceID(0)].VertexUVs.Num() : 1);
 
-	const FPolygonGroupArray& PolygonGroups = MeshDescription->PolygonGroups();
 	// Set up index buffer
-	for (const FPolygonGroupID& PolygonGroupID : PolygonGroups.GetElementIDs())
+	for (const FPolygonGroupID& PolygonGroupID : PolygonGroupArray.GetElementIDs())
 	{
-		TArray<FPolygonID> Polygons;
-		GetPolygonGroupTriangles(MeshDescription, Polygons, PolygonGroupID);
-
-		const FMeshPolygonGroup& PolygonGroup = PolygonGroups[PolygonGroupID];
+		const FMeshPolygonGroup& PolygonGroup = PolygonGroupArray[PolygonGroupID];
+		const TArray<FPolygonID> &Polygons = PolygonGroup.Polygons;
 		
 		TArray<uint32>& SectionIndices = OutPerSectionIndices[PolygonGroupID.GetValue()];
 
@@ -408,12 +398,12 @@ void BuildVertexBuffer(
 			StaticMesh->SectionInfoMap.Set(LodIndex, SectionIndex, SectionInfo);
 		}
 
-		for (FPolygonID& PolygonID : Polygons)
+		for (const FPolygonID& PolygonID : Polygons)
 		{
-			const FMeshPolygon& Polygon = MeshDescription->GetPolygon(PolygonID);
-			int32 ReseveSize = IndexBuffer.Num() + Polygon.Triangles.Num() * 3;
-			IndexBuffer.Reserve(ReseveSize);
-			OutWedgeMap.Reserve(ReseveSize);
+			const FMeshPolygon& Polygon = PolygonArray[PolygonID];
+			int32 ReserveSize = IndexBuffer.Num() + Polygon.Triangles.Num() * 3;
+			IndexBuffer.Reserve(ReserveSize);
+			OutWedgeMap.Reserve(ReserveSize);
 			uint32 MinIndex = TNumericLimits< uint32 >::Max();
 			uint32 MaxIndex = TNumericLimits< uint32 >::Min();
 			for (int32 TriangleIndex = 0; TriangleIndex < Polygon.Triangles.Num(); ++TriangleIndex)
@@ -421,8 +411,9 @@ void BuildVertexBuffer(
 				const FMeshTriangle& Triangle = Polygon.Triangles[TriangleIndex];
 				for (int32 TriVert = 0; TriVert < 3; ++TriVert)
 				{
-					int32 VertexInstanceValue = Triangle.GetVertexInstanceID(TriVert).GetValue();
-					const FMeshVertexInstance& VertexInstance = MeshDescription->GetVertexInstance(Triangle.GetVertexInstanceID(TriVert));
+					const FVertexInstanceID& VertexInstanceID = Triangle.GetVertexInstanceID(TriVert);
+					int32 VertexInstanceValue = VertexInstanceID.GetValue();
+					const FMeshVertexInstance& VertexInstance = VertexInstances[VertexInstanceID];
 					if (VertexInstance.Color != FLinearColor::White)
 					{
 						bHasColor = true;
@@ -436,17 +427,10 @@ void BuildVertexBuffer(
 					StaticMeshVertex.TangentY = ScaleMatrix.TransformVector(FVector::CrossProduct(VertexInstance.Normal, VertexInstance.Tangent).GetSafeNormal() * VertexInstance.BinormalSign).GetSafeNormal();
 					StaticMeshVertex.TangentZ = ScaleMatrix.TransformVector(VertexInstance.Normal).GetSafeNormal();
 					StaticMeshVertex.Color = VertexInstance.Color.ToFColor(true);
-					int32 NumTexCoords = FMath::Min<int32>(MAX_MESH_TEXTURE_COORDS, MAX_STATIC_TEXCOORDS);
-					for (int32 UVIndex = 0; UVIndex < NumTexCoords; ++UVIndex)
+					int32 MaxNumTexCoords = FMath::Min<int32>(MAX_MESH_TEXTURE_COORDS, MAX_STATIC_TEXCOORDS);
+					for (int32 UVIndex = 0; UVIndex < MaxNumTexCoords; ++UVIndex)
 					{
-						if(VertexInstance.VertexUVs.IsValidIndex(UVIndex))
-						{
-							StaticMeshVertex.UVs[UVIndex] = VertexInstance.VertexUVs[UVIndex];
-						}
-						else
-						{
-							StaticMeshVertex.UVs[UVIndex] = FVector2D(0.0f, 0.0f);
-						}
+						StaticMeshVertex.UVs[UVIndex] = VertexInstance.VertexUVs.IsValidIndex(UVIndex) ? VertexInstance.VertexUVs[UVIndex] : FVector2D(0.0f, 0.0f);
 					}
 					
 
@@ -490,10 +474,10 @@ void BuildVertexBuffer(
 	}
 
 	//Optimize before setting the buffer
-	if (MeshDescription->VertexInstances().Num() < 100000 * 3)
+	if (VertexInstances.Num() < 100000 * 3)
 	{
 		BuildOptimizationHelper::CacheOptimizeVertexAndIndexBuffer(StaticMeshBuildVertices, OutPerSectionIndices, OutWedgeMap);
-		check(OutWedgeMap.Num() == MeshDescription->VertexInstances().Num());
+		check(OutWedgeMap.Num() == VertexInstances.Num());
 	}
 
 	StaticMeshLOD.PositionVertexBuffer.Init(StaticMeshBuildVertices);
