optimize the create normals

diff --git a/Engine/Source/Runtime/MeshBuilder/Private/MeshDescriptionHelper.cpp b/Engine/Source/Runtime/MeshBuilder/Private/MeshDescriptionHelper.cpp
--- a/Engine/Source/Runtime/MeshBuilder/Private/MeshDescriptionHelper.cpp
+++ b/Engine/Source/Runtime/MeshBuilder/Private/MeshDescriptionHelper.cpp
@@ -56,8 +56,26 @@ UMeshDescription* FMeshDescriptionHelper::GetRenderMeshDescription(UObject* Owne
 
 	FVertexInstanceArray& VertexInstanceArray = RenderMeshDescription->VertexInstances();
 
-	if (BuildSettings->bRecomputeTangents || BuildSettings->bRecomputeNormals)
+	// Find overlapping corners to accelerate adjacency.
+	FindOverlappingCorners(OverlappingCorners, RenderMeshDescription, ComparisonThreshold);
+
+	// Compute any missing normals or tangents.
 	{
+		// Static meshes always blend normals of overlapping corners.
+		uint32 TangentOptions = FMeshDescriptionHelper::ETangentOptions::BlendOverlappingNormals;
+		if (BuildSettings->bRemoveDegenerates)
+		{
+			// If removing degenerate triangles, ignore them when computing tangents.
+			TangentOptions |= FMeshDescriptionHelper::ETangentOptions::IgnoreDegenerateTriangles;
+		}
+
+		//This function make sure the Polygon NTB are compute and also remove degenerated triangle from the render mesh description.
+		CreatePolygonNTB(RenderMeshDescription, BuildSettings->bRemoveDegenerates ? SMALL_NUMBER : 0.0f);
+		
+		//Keep the original mesh description NTBs if we do not rebuild the normals or tangents.
+		bool bComputeTangentLegacy = !BuildSettings->bUseMikkTSpace && (BuildSettings->bRecomputeNormals || BuildSettings->bRecomputeTangents);
+		bool bHasAllNormals = true;
+		bool bHasAllTangents = true;
 		for (const FVertexInstanceID& VertexInstanceID : VertexInstanceArray.GetElementIDs())
 		{
 			FMeshVertexInstance& VertexInstance = VertexInstanceArray[VertexInstanceID];
@@ -73,36 +91,22 @@ UMeshDescription* FMeshDescriptionHelper::GetRenderMeshDescription(UObject* Owne
 				//Dump the normals
 				VertexInstance.Normal = FVector(0.0f);
 			}
+			bHasAllNormals &= !VertexInstance.Normal.IsNearlyZero();
+			bHasAllTangents &= !VertexInstance.Tangent.IsNearlyZero();
 		}
-	}
 
-	// Compute any missing normals or tangents.
-	{
-		// Static meshes always blend normals of overlapping corners.
-		uint32 TangentOptions = FMeshDescriptionHelper::ETangentOptions::BlendOverlappingNormals;
-		if (BuildSettings->bRemoveDegenerates)
-		{
-			// If removing degenerate triangles, ignore them when computing tangents.
-			TangentOptions |= FMeshDescriptionHelper::ETangentOptions::IgnoreDegenerateTriangles;
-		}
 
-		//This function make sure the Polygon NTB are compute and also remove degenerated triangle from the render mesh description.
-		CreatePolygonNTB(RenderMeshDescription, BuildSettings->bRemoveDegenerates ? SMALL_NUMBER : 0.0f);
-		
-		// Find overlapping corners to accelerate adjacency.
-		FindOverlappingCorners(OverlappingCorners, RenderMeshDescription, ComparisonThreshold);
-			
-		//Keep the original mesh description NTBs if we do not rebuild the normals or tangents.
-		bool bComputeTangentLegacy = !BuildSettings->bUseMikkTSpace && (BuildSettings->bRecomputeNormals || BuildSettings->bRecomputeTangents);
-			
 		//MikkTSpace should be use only when the user want to recompute the normals or tangents otherwise should always fallback on builtin
 		//We cannot use mikkt space with degenerated normals fallback on buitin.
 		if (BuildSettings->bUseMikkTSpace && (BuildSettings->bRecomputeNormals || BuildSettings->bRecomputeTangents))
 		{
-			CreateNormals(RenderMeshDescription, (FMeshDescriptionHelper::ETangentOptions)TangentOptions, false);
+			if (!bHasAllNormals)
+			{
+				CreateNormals(RenderMeshDescription, (FMeshDescriptionHelper::ETangentOptions)TangentOptions, false);
+			}
 			CreateMikktTangents(RenderMeshDescription, (FMeshDescriptionHelper::ETangentOptions)TangentOptions);
 		}
-		else
+		else if(!bHasAllNormals || !bHasAllTangents)
 		{
 			//Set the compute tangent to true when we do not build using mikkt space
 			CreateNormals(RenderMeshDescription, (FMeshDescriptionHelper::ETangentOptions)TangentOptions, true);
@@ -152,6 +156,7 @@ void FMeshDescriptionHelper::ReduceLOD(const UMeshDescription* BaseMesh, UMeshDe
 
 void FMeshDescriptionHelper::CopyMeshDescription(UMeshDescription* SourceMeshDescription, UMeshDescription* DestinationMeshDescription) const
 {
+	BuildStatisticManager::FBuildStatisticScope StatScope(TEXT("CopyMeshDescription took"));
 	//Copy the Source into the destination
 	//Save the Source
 	TArray<uint8> TempBytes;
@@ -232,6 +237,11 @@ void FMeshDescriptionHelper::CreatePolygonNTB(UMeshDescription* MeshDescription,
 		FVector TangentZ(0.0f);
 
 		FMeshPolygon& Polygon = PolygonArray[PolygonID];
+		if (!Polygon.PolygonNormal.IsNearlyZero())
+		{
+			//By pass normal calculation if its already done
+			continue;
+		}
 		const TArray<FMeshTriangle>& MeshTriangles = Polygon.Triangles;
 #ifdef ENABLE_NTB_CHECK
 		//Assume triangle are build
@@ -284,13 +294,9 @@ void FMeshDescriptionHelper::CreatePolygonNTB(UMeshDescription* MeshDescription,
 		{
 			DegeneratePolygons.Add(PolygonID);
 		}
-		//Set the polygon Normal value only once if it is not set.
-		if (Polygon.PolygonNormal.IsNearlyZero())
-		{
-			Polygon.PolygonTangent = TangentX;
-			Polygon.PolygonBinormal = TangentY;
-			Polygon.PolygonNormal = TangentZ;
-		}
+		Polygon.PolygonTangent = TangentX;
+		Polygon.PolygonBinormal = TangentY;
+		Polygon.PolygonNormal = TangentZ;
 	}
 
 	//Delete the degenerated polygons. The array is fill only if the remove degenerated option is turn on.
@@ -335,29 +341,29 @@ struct FVertexInfo
 		VertexInstanceID = FVertexInstanceID::Invalid;
 		UVs = FVector2D(0.0f, 0.0f);
 		EdgeIDs.Reserve(2);//Most of the time a edge has two triangles
+		Polygon = nullptr;
 	}
 
 	FPolygonID PolygonID;
+	FMeshPolygon* Polygon;
 	FVertexInstanceID VertexInstanceID;
 	FVector2D UVs;
 	TArray<FEdgeID> EdgeIDs;
+	TMap<FEdgeID, FMeshEdge*> Edges;
 };
 
 void RecursiveFillPolygonGroup(
-	  UMeshDescription* MeshDescription
-	, TMap<FPolygonID, FVertexInfo>& VertexInfoMap
-	, TSet<FPolygonID>& CurrentGroup
+	  TMap<FPolygonID, FVertexInfo>& VertexInfoMap
+	, TArray<FPolygonID>& CurrentGroup
 	, FVertexInfo& CurrentVertexInfo
-	, TSet<FPolygonID>& ConsumedPolygon)
+	, TArray<FPolygonID>& ConsumedPolygon)
 {
-	CurrentGroup.Add(CurrentVertexInfo.PolygonID);
-	ConsumedPolygon.Add(CurrentVertexInfo.PolygonID);
-
-	FEdgeArray& EdgeArray = MeshDescription->Edges();
+	CurrentGroup.AddUnique(CurrentVertexInfo.PolygonID);
+	ConsumedPolygon.AddUnique(CurrentVertexInfo.PolygonID);
 
 	for (FEdgeID &EdgeID : CurrentVertexInfo.EdgeIDs)
 	{
-		FMeshEdge& Edge = EdgeArray[EdgeID];
+		FMeshEdge& Edge = *(CurrentVertexInfo.Edges[EdgeID]);
 		if (Edge.bIsHardEdge)
 		{
 			//End of the group
@@ -374,7 +380,7 @@ void RecursiveFillPolygonGroup(
 			//Do not repeat polygons
 			if (!ConsumedPolygon.Contains(OtherVertexInfo.PolygonID))
 			{
-				RecursiveFillPolygonGroup(MeshDescription, VertexInfoMap, CurrentGroup, OtherVertexInfo, ConsumedPolygon);
+				RecursiveFillPolygonGroup(VertexInfoMap, CurrentGroup, OtherVertexInfo, ConsumedPolygon);
 			}
 		}
 	}
@@ -403,6 +409,7 @@ void FMeshDescriptionHelper::CreateNormals(UMeshDescription* MeshDescription, FM
 	FEdgeArray& EdgeArray = MeshDescription->Edges();
 
 	TMap<FPolygonID, FVertexInfo> VertexInfoMap;
+	TMap<FVertexInstanceID, FMeshVertexInstance*> VertexInstanceMap;
 	VertexInfoMap.Reserve(20);
 	//Iterate all vertex to compute normals for all vertex instance
 	for (const FVertexID& VertexID : VertexArray.GetElementIDs())
@@ -411,6 +418,7 @@ void FMeshDescriptionHelper::CreateNormals(UMeshDescription* MeshDescription, FM
 
 		VertexInfoMap.Reset();
 		
+		bool bPointHasAllTangents = true;
 		//Fill the VertexInfoMap
 		for (const FEdgeID& EdgeID : Vertex.ConnectedEdgeIDs)
 		{
@@ -419,12 +427,14 @@ void FMeshDescriptionHelper::CreateNormals(UMeshDescription* MeshDescription, FM
 			{
 				FVertexInfo& VertexInfo = VertexInfoMap.FindOrAdd(PolygonID);
 				int32 EdgeIndex = VertexInfo.EdgeIDs.AddUnique(EdgeID);
+				VertexInfo.Edges.Add(EdgeID, &Edge);
 
 				if (VertexInfo.PolygonID == FPolygonID::Invalid)
 				{
 					FMeshPolygon& Polygon = PolygonArray[PolygonID];
 
 					VertexInfo.PolygonID = PolygonID;
+					VertexInfo.Polygon = &Polygon;
 					for (FVertexInstanceID& VertexInstanceID : Polygon.PerimeterContour.VertexInstanceIDs)
 					{
 						FMeshVertexInstance& VertexInstance = VertexInstanceArray[VertexInstanceID];
@@ -432,6 +442,8 @@ void FMeshDescriptionHelper::CreateNormals(UMeshDescription* MeshDescription, FM
 						{
 							VertexInfo.VertexInstanceID = VertexInstanceID;
 							VertexInfo.UVs = VertexInstance.VertexUVs[0];
+							VertexInstanceMap.Add(VertexInstanceID, &VertexInstance);
+							bPointHasAllTangents &= !VertexInstance.Normal.IsNearlyZero() && !VertexInstance.Tangent.IsNearlyZero();
 							break;
 						}
 					}
@@ -439,12 +451,17 @@ void FMeshDescriptionHelper::CreateNormals(UMeshDescription* MeshDescription, FM
 			}
 		}
 
+		if (bPointHasAllTangents)
+		{
+			continue;
+		}
+
 		//Make sure we consume all our vertex instance
 		check(VertexInfoMap.Num() == Vertex.VertexInstanceIDs.Num());
 
 		//Build all group by recursively traverse all polygon connected to the vertex
-		TArray<TSet<FPolygonID>> Groups;
-		TSet<FPolygonID> ConsumedPolygon;
+		TArray<TArray<FPolygonID>> Groups;
+		TArray<FPolygonID> ConsumedPolygon;
 		for (auto Kvp : VertexInfoMap)
 		{
 			if (ConsumedPolygon.Contains(Kvp.Key))
@@ -454,12 +471,12 @@ void FMeshDescriptionHelper::CreateNormals(UMeshDescription* MeshDescription, FM
 
 			int32 CurrentGroupIndex = Groups.AddZeroed();
 			FVertexInfo& VertexInfo = VertexInfoMap.FindOrAdd(Kvp.Key);
-			RecursiveFillPolygonGroup(MeshDescription, VertexInfoMap, Groups[CurrentGroupIndex], VertexInfo, ConsumedPolygon);
+			RecursiveFillPolygonGroup(VertexInfoMap, Groups[CurrentGroupIndex], VertexInfo, ConsumedPolygon);
 		}
 
 		//Smooth every connected group
 		ConsumedPolygon.Reset();
-		for(TSet<FPolygonID> Group : Groups)
+		for(const TArray<FPolygonID>& Group : Groups)
 		{
 			//Compute tangents data
 			TMap<FVector2D, FVector> GroupTangent;
@@ -467,14 +484,14 @@ void FMeshDescriptionHelper::CreateNormals(UMeshDescription* MeshDescription, FM
 
 			TArray<FVertexInstanceID> VertexInstanceInGroup;
 			FVector GroupNormal(0.0f);
-			for(FPolygonID& PolygonID : Group)
+			for(const FPolygonID& PolygonID : Group)
 			{
 #ifdef ENABLE_NTB_CHECK
 				check(!ConsumedPolygon.Contains(PolygonID));
 #endif
 				ConsumedPolygon.Add(PolygonID);
 				VertexInstanceInGroup.Add(VertexInfoMap[PolygonID].VertexInstanceID);
-				FMeshPolygon& Polygon = PolygonArray[PolygonID];
+				FMeshPolygon& Polygon = *(VertexInfoMap[PolygonID].Polygon);
 				GroupNormal += Polygon.PolygonNormal;
 				if (bComputeTangent)
 				{
@@ -506,7 +523,7 @@ void FMeshDescriptionHelper::CreateNormals(UMeshDescription* MeshDescription, FM
 			//Apply the average NTB on all Vertex instance
 			for (FVertexInstanceID& VertexInstanceID : VertexInstanceInGroup)
 			{
-				FMeshVertexInstance& VertexInstance = VertexInstanceArray[VertexInstanceID];
+				FMeshVertexInstance& VertexInstance = *(VertexInstanceMap[VertexInstanceID]);
 				FVector2D& VertexUV = VertexInstance.VertexUVs[0];
 			
 				if (VertexInstance.Normal.IsNearlyZero(SMALL_NUMBER))
@@ -558,7 +575,6 @@ void FMeshDescriptionHelper::CreateNormals(UMeshDescription* MeshDescription, FM
 					VertexInstance.Tangent = GroupTangentValue;
 					//If the BiNormal is zero set the sign to 1.0f
 					VertexInstance.BinormalSign = GetBasisDeterminantSign(GroupTangentValue, GroupBiNormalValue, VertexInstance.Normal);
-
 				}
 			}
 		}
@@ -569,7 +585,6 @@ void FMeshDescriptionHelper::CreateMikktTangents(UMeshDescription* MeshDescripti
 {
 	BuildStatisticManager::FBuildStatisticScope StatScope(TEXT("CreateMikktTangents took"));
 	bool bIgnoreDegenerateTriangles = (TangentOptions & FMeshDescriptionHelper::ETangentOptions::IgnoreDegenerateTriangles) != 0;
-	float ComparisonThreshold = bIgnoreDegenerateTriangles ? THRESH_POINTS_ARE_SAME : 0.0f;
 
 	// we can use mikktspace to calculate the tangents
 	SMikkTSpaceInterface MikkTInterface;
diff --git a/Engine/Source/Runtime/MeshBuilder/Private/StaticMeshBuilder.cpp b/Engine/Source/Runtime/MeshBuilder/Private/StaticMeshBuilder.cpp
--- a/Engine/Source/Runtime/MeshBuilder/Private/StaticMeshBuilder.cpp
+++ b/Engine/Source/Runtime/MeshBuilder/Private/StaticMeshBuilder.cpp
@@ -84,11 +84,9 @@ bool FStaticMeshBuilder::Build(UStaticMesh* StaticMesh, const FStaticMeshLODGrou
 		float VertexComparisonThreshold = LODBuildSettings.bRemoveDegenerates ? THRESH_POINTS_ARE_SAME : 0.0f;
 
 		//Build new vertex buffers
-		static TArray< FStaticMeshBuildVertex > StaticMeshBuildVertices;
-		StaticMeshBuildVertices.Reset();
+		TArray< FStaticMeshBuildVertex > StaticMeshBuildVertices;
 
-		static TArray< uint32 > IndexBuffer;
-		IndexBuffer.Reset();
+		TArray< uint32 > IndexBuffer;
 
 		StaticMeshLOD.Sections.Empty(PolygonGroups.Num());
 		TArray<int32> RemapVerts; //Because we will remove MeshVertex that are redundant, we need a remap
@@ -437,7 +435,6 @@ void BuildVertexBuffer(
 					//Never add duplicated vertex instance
 					DupVerts.Reset();
 					OverlappingCorners.MultiFind(VertexInstanceValue, DupVerts);
-					DupVerts.Sort();
 					int32 Index = INDEX_NONE;
 					for (int32 k = 0; k < DupVerts.Num(); k++)
 					{
@@ -480,10 +477,10 @@ void BuildVertexBuffer(
 		check(OutWedgeMap.Num() == VertexInstances.Num());
 	}
 
-	StaticMeshLOD.PositionVertexBuffer.Init(StaticMeshBuildVertices);
 	StaticMeshLOD.VertexBuffer.SetUseHighPrecisionTangentBasis(LODBuildSettings.bUseHighPrecisionTangentBasis);
 	StaticMeshLOD.VertexBuffer.SetUseFullPrecisionUVs(LODBuildSettings.bUseFullPrecisionUVs);
 	StaticMeshLOD.VertexBuffer.Init(StaticMeshBuildVertices, NumTextureCoord);
+	StaticMeshLOD.PositionVertexBuffer.Init(StaticMeshBuildVertices);
 	if (bHasColor)
 	{
 		StaticMeshLOD.ColorVertexBuffer.Init(StaticMeshBuildVertices);
