Total revamp of mesh element attribute model.

diff --git a/Engine/Source/Editor/UnrealEd/Private/Fbx/FbxStaticMeshImport.cpp b/Engine/Source/Editor/UnrealEd/Private/Fbx/FbxStaticMeshImport.cpp
--- a/Engine/Source/Editor/UnrealEd/Private/Fbx/FbxStaticMeshImport.cpp
+++ b/Engine/Source/Editor/UnrealEd/Private/Fbx/FbxStaticMeshImport.cpp
@@ -34,6 +34,7 @@
 #include "Misc/FbxErrors.h"
 #include "PhysicsEngine/BodySetup.h"
 #include "MeshDescription.h"
+#include "MeshAttributes.h"
 #include "IMeshBuilderModule.h"
 #include "Settings/EditorExperimentalSettings.h"
 
@@ -521,6 +522,23 @@ bool UnFbx::FFbxImporter::BuildStaticMeshFromGeometry(FbxNode* Node, UStaticMesh
 	bool bHasNonDegenerateTriangles = false;
 	if(MeshDescription != nullptr)
 	{
+		TVertexAttributeArray<FVector>& VertexPositions = MeshDescription->VertexAttributes().GetAttributes<FVector>(MeshAttribute::Vertex::Position);
+
+		TVertexInstanceAttributeArray<FVector>& VertexInstanceNormals = MeshDescription->VertexInstanceAttributes().GetAttributes<FVector>(MeshAttribute::VertexInstance::Normal);
+		TVertexInstanceAttributeArray<FVector>& VertexInstanceTangents = MeshDescription->VertexInstanceAttributes().GetAttributes<FVector>(MeshAttribute::VertexInstance::Tangent);
+		TVertexInstanceAttributeArray<float>& VertexInstanceBinormalSigns = MeshDescription->VertexInstanceAttributes().GetAttributes<float>(MeshAttribute::VertexInstance::BinormalSign);
+		TVertexInstanceAttributeArray<FVector4>& VertexInstanceColors = MeshDescription->VertexInstanceAttributes().GetAttributes<FVector4>(MeshAttribute::VertexInstance::Color);
+		TVertexInstanceAttributeIndicesArray<FVector2D>& VertexInstanceUVs = MeshDescription->VertexInstanceAttributes().GetAttributesSet<FVector2D>(MeshAttribute::VertexInstance::TextureCoordinate);
+
+		TEdgeAttributeArray<bool>& EdgeHardnesses = MeshDescription->EdgeAttributes().GetAttributes<bool>(MeshAttribute::Edge::IsHard);
+		TEdgeAttributeArray<float>& EdgeCreaseSharpnesses = MeshDescription->EdgeAttributes().GetAttributes<float>(MeshAttribute::Edge::CreaseSharpness);
+
+		TPolygonGroupAttributeArray<UObject*>& PolygonGroupMaterialAssets = MeshDescription->PolygonGroupAttributes().GetAttributes<UObject*>(MeshAttribute::PolygonGroup::MaterialAsset);
+		TPolygonGroupAttributeArray<FName>& PolygonGroupMaterialSlotNames = MeshDescription->PolygonGroupAttributes().GetAttributes<FName>(MeshAttribute::PolygonGroup::MaterialSlotName);
+		TPolygonGroupAttributeArray<FName>& PolygonGroupImportedMaterialSlotNames = MeshDescription->PolygonGroupAttributes().GetAttributes<FName>(MeshAttribute::PolygonGroup::ImportedMaterialSlotName);
+		TPolygonGroupAttributeArray<bool>& PolygonGroupCollision = MeshDescription->PolygonGroupAttributes().GetAttributes<bool>(MeshAttribute::PolygonGroup::EnableCollision);
+		TPolygonGroupAttributeArray<bool>& PolygonGroupCastShadow = MeshDescription->PolygonGroupAttributes().GetAttributes<bool>(MeshAttribute::PolygonGroup::CastShadow);
+
 		int32 VertexOffset = MeshDescription->Vertices().Num();
 		int32 VertexInstanceOffset = MeshDescription->VertexInstances().Num();
 		int32 TriangleOffset = MeshDescription->Polygons().Num();
@@ -529,25 +547,14 @@ bool UnFbx::FFbxImporter::BuildStaticMeshFromGeometry(FbxNode* Node, UStaticMesh
 		TMap<int32, FPolygonGroupID> PolygonGroupMapping;
 
 		// When importing multiple mesh pieces to the same static mesh.  Ensure each mesh piece has the same number of Uv's
-		int32 ExistingUVCount = 0;
-		for (const FVertexInstanceID& VertexInstanceID : MeshDescription->VertexInstances().GetElementIDs())
-		{
-			ExistingUVCount = FMath::Max(ExistingUVCount, MeshDescription->GetVertexInstance(VertexInstanceID).VertexUVs.Num());
-		}
+		int32 ExistingUVCount = VertexInstanceUVs.GetNumIndices();
 		
 		int32 NumUVs = FMath::Max(FBXUVs.UniqueUVCount, ExistingUVCount);
 		// At least one UV set must exist.  
 		NumUVs = FMath::Max(1, NumUVs);
 
 		//Make sure all Vertex instance have the correct number of UVs
-		for (const FVertexInstanceID& VertexInstanceID : MeshDescription->VertexInstances().GetElementIDs())
-		{
-			FMeshVertexInstance& VertexInstance = MeshDescription->GetVertexInstance(VertexInstanceID);
-			while (VertexInstance.VertexUVs.Num() < NumUVs)
-			{
-				VertexInstance.VertexUVs.Add(FVector2D(0.0f, 0.0f));
-			}
-		}
+		VertexInstanceUVs.SetNumIndices( NumUVs );
 
 		//Fill the vertex array
 		for (int32 VertexIndex = 0; VertexIndex < VertexCount; ++VertexIndex)
@@ -558,9 +565,7 @@ bool UnFbx::FFbxImporter::BuildStaticMeshFromGeometry(FbxNode* Node, UStaticMesh
 			const FVector VertexPosition = Converter.ConvertPos(FbxPosition);
 			
 			FVertexID AddedVertexId = MeshDescription->CreateVertex();
-			FMeshVertex& NewVertex = MeshDescription->GetVertex(AddedVertexId);
-			NewVertex.VertexPosition = VertexPosition;
-			NewVertex.CornerSharpness = 0.0f;
+			VertexPositions[AddedVertexId] = VertexPosition;
 			if (AddedVertexId.GetValue() != RealVertexIndex)
 			{
 				AddTokenizedErrorMessage(FTokenizedMessage::Create(EMessageSeverity::Error, FText::Format(LOCTEXT("Error_CannotCreateVertex", "Cannot create valid vertex for mesh '{0}'"), FText::FromString(Mesh->GetName()))), FFbxErrors::StaticMesh_BuildError);
@@ -581,9 +586,9 @@ bool UnFbx::FFbxImporter::BuildStaticMeshFromGeometry(FbxNode* Node, UStaticMesh
 				int32 VertexInstanceIndex = RealTriangleIndex * 3 + CornerIndex;
 				const FVertexInstanceID VertexInstanceID(VertexInstanceIndex);
 				CornerInstanceIDs[CornerIndex] = VertexInstanceID;
-				int32 ControlPointIndex = Mesh->GetPolygonVertex(TriangleIndex, CornerIndex);
+				const int32 ControlPointIndex = Mesh->GetPolygonVertex(TriangleIndex, CornerIndex);
 				const FVertexID VertexID(VertexOffset + ControlPointIndex);
-				FVector VertexPosition = MeshDescription->GetVertex(VertexID).VertexPosition;
+				const FVector VertexPosition = VertexPositions[VertexID];
 				CornerVerticesIDs[CornerIndex] = VertexID;
 
 				FVertexInstanceID AddedVertexInstanceId = MeshDescription->CreateVertexInstance(VertexID);
@@ -591,7 +596,6 @@ bool UnFbx::FFbxImporter::BuildStaticMeshFromGeometry(FbxNode* Node, UStaticMesh
 				//Make sure the Added vertex instance ID is matching the expected vertex instance ID
 				check(AddedVertexInstanceId == VertexInstanceID);
 				
-				FMeshVertexInstance& NewVertexInstance = MeshDescription->GetVertexInstance(AddedVertexInstanceId);
 				if (AddedVertexInstanceId.GetValue() != VertexInstanceIndex)
 				{
 					AddTokenizedErrorMessage(FTokenizedMessage::Create(EMessageSeverity::Error, FText::Format(LOCTEXT("Error_CannotCreateVertexInstance", "Cannot create valid vertex instance for mesh '{0}'"), FText::FromString(Mesh->GetName()))), FFbxErrors::StaticMesh_BuildError);
@@ -599,28 +603,20 @@ bool UnFbx::FFbxImporter::BuildStaticMeshFromGeometry(FbxNode* Node, UStaticMesh
 				}
 
 				//UVs attributes
-				if (FBXUVs.UniqueUVCount > 0)
+				for (int32 UVLayerIndex = 0; UVLayerIndex < FBXUVs.UniqueUVCount; UVLayerIndex++)
 				{
-					for (int32 UVLayerIndex = 0; UVLayerIndex < FBXUVs.UniqueUVCount; UVLayerIndex++)
+					FVector2D FinalUVVector(0.0f, 0.0f);
+					if (FBXUVs.LayerElementUV[UVLayerIndex] != NULL)
 					{
-						FVector2D FinalUVVector(0.0f, 0.0f);
-						if (FBXUVs.LayerElementUV[UVLayerIndex] != NULL)
-						{
-							int UVMapIndex = (FBXUVs.UVMappingMode[UVLayerIndex] == FbxLayerElement::eByControlPoint) ? ControlPointIndex : TriangleIndex * 3 + CornerIndex;
-							int32 UVIndex = (FBXUVs.UVReferenceMode[UVLayerIndex] == FbxLayerElement::eDirect) ?
-								UVMapIndex : FBXUVs.LayerElementUV[UVLayerIndex]->GetIndexArray().GetAt(UVMapIndex);
+						int UVMapIndex = (FBXUVs.UVMappingMode[UVLayerIndex] == FbxLayerElement::eByControlPoint) ? ControlPointIndex : TriangleIndex * 3 + CornerIndex;
+						int32 UVIndex = (FBXUVs.UVReferenceMode[UVLayerIndex] == FbxLayerElement::eDirect) ?
+							UVMapIndex : FBXUVs.LayerElementUV[UVLayerIndex]->GetIndexArray().GetAt(UVMapIndex);
 
-							FbxVector2	UVVector = FBXUVs.LayerElementUV[UVLayerIndex]->GetDirectArray().GetAt(UVIndex);
-							FinalUVVector.X = static_cast<float>(UVVector[0]);
-							FinalUVVector.Y = 1.f - static_cast<float>(UVVector[1]);   //flip the Y of UVs for DirectX
-						}
-						NewVertexInstance.VertexUVs.Add(FinalUVVector);
+						FbxVector2	UVVector = FBXUVs.LayerElementUV[UVLayerIndex]->GetDirectArray().GetAt(UVIndex);
+						FinalUVVector.X = static_cast<float>(UVVector[0]);
+						FinalUVVector.Y = 1.f - static_cast<float>(UVVector[1]);   //flip the Y of UVs for DirectX
 					}
-				}
-				else
-				{
-					FVector2D FinalUVVector(0.0f, 0.0f);
-					NewVertexInstance.VertexUVs.Add(FinalUVVector);
+					VertexInstanceUVs.GetArrayForIndex(UVLayerIndex)[AddedVertexInstanceId] = FinalUVVector;
 				}
 
 				//Color attribute
@@ -642,7 +638,7 @@ bool UnFbx::FFbxImporter::BuildStaticMeshFromGeometry(FbxNode* Node, UStaticMesh
 							uint8(255.f*VertexColor.mBlue),
 							uint8(255.f*VertexColor.mAlpha)
 						);
-						NewVertexInstance.Color = FLinearColor(VertexInstanceColor);
+						VertexInstanceColors[AddedVertexInstanceId] = FVector4(FLinearColor(VertexInstanceColor));
 					}
 				}
 				else if (VertexColorImportOption == EVertexColorImportOption::Ignore)
@@ -653,14 +649,14 @@ bool UnFbx::FFbxImporter::BuildStaticMeshFromGeometry(FbxNode* Node, UStaticMesh
 					if (PaintedColor)
 					{
 						// A matching color for this vertex was found
-						NewVertexInstance.Color = FLinearColor(*PaintedColor);
+						VertexInstanceColors[AddedVertexInstanceId] = FVector4(FLinearColor(*PaintedColor));
 					}
 				}
 				else
 				{
 					// set the triangle's vertex color to a constant override
 					check(VertexColorImportOption == EVertexColorImportOption::Override);
-					NewVertexInstance.Color = FLinearColor(VertexOverrideColor);
+					VertexInstanceColors[AddedVertexInstanceId] = FVector4(FLinearColor(VertexOverrideColor));
 				}
 
 				if (LayerElementNormal)
@@ -675,7 +671,7 @@ bool UnFbx::FFbxImporter::BuildStaticMeshFromGeometry(FbxNode* Node, UStaticMesh
 					FbxVector4 TempValue = LayerElementNormal->GetDirectArray().GetAt(NormalValueIndex);
 					TempValue = TotalMatrixForNormal.MultT(TempValue);
 					FVector TangentZ = Converter.ConvertDir(TempValue);
-					NewVertexInstance.Normal = TangentZ.GetSafeNormal();
+					VertexInstanceNormals[AddedVertexInstanceId] = TangentZ.GetSafeNormal();
 					//tangents and binormals share the same reference, mapping mode and index array
 					if (bHasNTBInformation)
 					{
@@ -687,7 +683,7 @@ bool UnFbx::FFbxImporter::BuildStaticMeshFromGeometry(FbxNode* Node, UStaticMesh
 						TempValue = LayerElementTangent->GetDirectArray().GetAt(TangentValueIndex);
 						TempValue = TotalMatrixForNormal.MultT(TempValue);
 						FVector TangentX = Converter.ConvertDir(TempValue);
-						NewVertexInstance.Tangent = TangentX.GetSafeNormal();
+						VertexInstanceTangents[AddedVertexInstanceId] = TangentX.GetSafeNormal();
 
 						int BinormalMapIndex = (BinormalMappingMode == FbxLayerElement::eByControlPoint) ?
 							ControlPointIndex : TriangleCornerIndex;
@@ -697,7 +693,7 @@ bool UnFbx::FFbxImporter::BuildStaticMeshFromGeometry(FbxNode* Node, UStaticMesh
 						TempValue = LayerElementBinormal->GetDirectArray().GetAt(BinormalValueIndex);
 						TempValue = TotalMatrixForNormal.MultT(TempValue);
 						FVector TangentY = -Converter.ConvertDir(TempValue);
-						NewVertexInstance.BinormalSign = GetBasisDeterminantSign(TangentX.GetSafeNormal(), TangentY.GetSafeNormal(), TangentZ.GetSafeNormal());
+						VertexInstanceBinormalSigns[AddedVertexInstanceId] = GetBasisDeterminantSign(TangentX.GetSafeNormal(), TangentY.GetSafeNormal(), TangentZ.GetSafeNormal());
 					}
 				}
 			}
@@ -707,9 +703,9 @@ bool UnFbx::FFbxImporter::BuildStaticMeshFromGeometry(FbxNode* Node, UStaticMesh
 			{
 				float TriangleComparisonThreshold = ImportOptions->bRemoveDegenerates ? THRESH_POINTS_ARE_SAME : 0.0f;
 				FVector VertexPosition[3];
-				VertexPosition[0] = MeshDescription->GetVertex(CornerVerticesIDs[0]).VertexPosition;
-				VertexPosition[1] = MeshDescription->GetVertex(CornerVerticesIDs[1]).VertexPosition;
-				VertexPosition[2] = MeshDescription->GetVertex(CornerVerticesIDs[2]).VertexPosition;
+				VertexPosition[0] = VertexPositions[CornerVerticesIDs[0]];
+				VertexPosition[1] = VertexPositions[CornerVerticesIDs[1]];
+				VertexPosition[2] = VertexPositions[CornerVerticesIDs[2]];
 				if (!( VertexPosition[0].Equals(VertexPosition[1], TriangleComparisonThreshold)
 					|| VertexPosition[0].Equals(VertexPosition[2], TriangleComparisonThreshold)
 					|| VertexPosition[1].Equals(VertexPosition[2], TriangleComparisonThreshold)))
@@ -758,7 +754,7 @@ bool UnFbx::FFbxImporter::BuildStaticMeshFromGeometry(FbxNode* Node, UStaticMesh
 				FPolygonGroupID ExistingPolygonGroup = FPolygonGroupID::Invalid;
 				for (const FPolygonGroupID PolygonGroupID : MeshDescription->PolygonGroups().GetElementIDs())
 				{
-					if (MeshDescription->PolygonGroups()[PolygonGroupID].ImportedMaterialSlotName == ImportedMaterialSlotName)
+					if (PolygonGroupImportedMaterialSlotNames[PolygonGroupID] == ImportedMaterialSlotName)
 					{
 						ExistingPolygonGroup = PolygonGroupID;
 						break;
@@ -767,12 +763,11 @@ bool UnFbx::FFbxImporter::BuildStaticMeshFromGeometry(FbxNode* Node, UStaticMesh
 				if (ExistingPolygonGroup == FPolygonGroupID::Invalid)
 				{
 					ExistingPolygonGroup = MeshDescription->CreatePolygonGroup();
-					FMeshPolygonGroup& NewPolygonGroup = MeshDescription->GetPolygonGroup(ExistingPolygonGroup);
-					NewPolygonGroup.bCastShadow = true;
-					NewPolygonGroup.bEnableCollision = bEnableCollision;
-					NewPolygonGroup.ImportedMaterialSlotName = ImportedMaterialSlotName;
-					NewPolygonGroup.MaterialSlotName = ImportedMaterialSlotName;
-					NewPolygonGroup.MaterialAsset = FStringAssetReference(Material);
+					PolygonGroupCastShadow[ExistingPolygonGroup] = true;
+					PolygonGroupCollision[ExistingPolygonGroup] = bEnableCollision;
+					PolygonGroupImportedMaterialSlotNames[ExistingPolygonGroup] = ImportedMaterialSlotName;
+					PolygonGroupMaterialSlotNames[ExistingPolygonGroup] = ImportedMaterialSlotName;
+					PolygonGroupMaterialAssets[ExistingPolygonGroup] = Material;
 				}
 				PolygonGroupMapping.Add(RealMaterialIndex, ExistingPolygonGroup);
 			}
@@ -794,29 +789,17 @@ bool UnFbx::FFbxImporter::BuildStaticMeshFromGeometry(FbxNode* Node, UStaticMesh
 					EdgeVertexIDs[0] = CornerVerticesIDs[CornerIndices[0]];
 					EdgeVertexIDs[1] = CornerVerticesIDs[CornerIndices[1]];
 
-					FEdgeID MatchEdgeId = FEdgeID::Invalid;
-					const FMeshVertex& EdgeMeshVertex = MeshDescription->GetVertex(EdgeVertexIDs[0]);
-					for (const FEdgeID& EdgeID : EdgeMeshVertex.ConnectedEdgeIDs)
-					{
-						FMeshEdge& Edge = MeshDescription->GetEdge(EdgeID);
-						if (EdgeVertexIDs[0] == Edge.VertexIDs[0] && EdgeVertexIDs[1] == Edge.VertexIDs[1] ||
-							EdgeVertexIDs[1] == Edge.VertexIDs[0] && EdgeVertexIDs[0] == Edge.VertexIDs[1])
-						{
-							MatchEdgeId = EdgeID;
-							break;
-						}
-					}
+					FEdgeID MatchEdgeId = MeshDescription->GetVertexPairEdge(EdgeVertexIDs[0], EdgeVertexIDs[1]);
 					if (MatchEdgeId == FEdgeID::Invalid)
 					{
 						MatchEdgeId = MeshDescription->CreateEdge(EdgeVertexIDs[0], EdgeVertexIDs[1]);
 					}
 					ContourPoint.EdgeID = MatchEdgeId;
 					ContourPoint.VertexInstanceID = CornerInstanceIDs[CornerIndices[0]];
-					FMeshEdge& ExistingEdge = MeshDescription->GetEdge(MatchEdgeId);
 
 					int32 EdgeIndex = Mesh->GetMeshEdgeIndexForPolygon(TriangleIndex, TriangleEdgeNumber);
-					ExistingEdge.CreaseSharpness = (float)Mesh->GetEdgeCreaseInfo(EdgeIndex);
-					if (!ExistingEdge.bIsHardEdge)
+					EdgeCreaseSharpnesses[MatchEdgeId] = (float)Mesh->GetEdgeCreaseInfo(EdgeIndex);
+					if (!EdgeHardnesses[MatchEdgeId])
 					{
 						if (bSmoothingAvailable && SmoothingInfo)
 						{
@@ -824,7 +807,7 @@ bool UnFbx::FFbxImporter::BuildStaticMeshFromGeometry(FbxNode* Node, UStaticMesh
 							{
 								int32 lSmoothingIndex = (SmoothingReferenceMode == FbxLayerElement::eDirect) ? EdgeIndex : SmoothingInfo->GetIndexArray().GetAt(EdgeIndex);
 								//Set the hard edges
-								ExistingEdge.bIsHardEdge = (SmoothingInfo->GetDirectArray().GetAt(lSmoothingIndex) == 0);
+								EdgeHardnesses[MatchEdgeId] = (SmoothingInfo->GetDirectArray().GetAt(lSmoothingIndex) == 0);
 							}
 							else
 							{
@@ -834,7 +817,7 @@ bool UnFbx::FFbxImporter::BuildStaticMeshFromGeometry(FbxNode* Node, UStaticMesh
 						else
 						{
 							//When there is no smoothing group we set all edge to hard (faceted mesh)
-							ExistingEdge.bIsHardEdge = true;
+							EdgeHardnesses[MatchEdgeId] = true;
 						}
 					}
 				}
@@ -842,13 +825,9 @@ bool UnFbx::FFbxImporter::BuildStaticMeshFromGeometry(FbxNode* Node, UStaticMesh
 
 			FPolygonGroupID PolygonGroupID = PolygonGroupMapping[RealMaterialIndex];
 			// Insert a polygon into the mesh
-			const FPolygonID NewPolygonID = MeshDescription->CreatePolygon(Contours);
-			FMeshPolygonGroup& MeshPolygonGroup = MeshDescription->GetPolygonGroup(PolygonGroupID);
-			FMeshPolygon& NewPolygon = MeshDescription->GetPolygon(NewPolygonID);
-			NewPolygon.PolygonGroupID = PolygonGroupID;
-			MeshPolygonGroup.Polygons.Add(NewPolygonID);
-			int32 NewTriangleIndex = NewPolygon.Triangles.AddDefaulted();
-			FMeshTriangle& NewTriangle = NewPolygon.Triangles[NewTriangleIndex];
+			const FPolygonID NewPolygonID = MeshDescription->CreatePolygon(PolygonGroupID, Contours);
+			int32 NewTriangleIndex = MeshDescription->GetPolygonTriangles(NewPolygonID).AddDefaulted();
+			FMeshTriangle& NewTriangle = MeshDescription->GetPolygonTriangles(NewPolygonID)[NewTriangleIndex];
 			for (int32 TriangleVertexIndex = 0; TriangleVertexIndex < 3; ++TriangleVertexIndex)
 			{
 				const FVertexInstanceID &VertexInstanceID = CornerInstanceIDs[TriangleVertexIndex];
@@ -1854,12 +1833,20 @@ UStaticMesh* UnFbx::FFbxImporter::ImportStaticMeshAsSingle(UObject* InParent, TA
 		}
 		else
 		{
+			TPolygonGroupAttributeArray<UObject*>& PolygonGroupMaterialAssets = MeshDescription->PolygonGroupAttributes().GetAttributes<UObject*>( MeshAttribute::PolygonGroup::MaterialAsset );
+			TPolygonGroupAttributeArray<FName>& PolygonGroupMaterialSlotNames = MeshDescription->PolygonGroupAttributes().GetAttributes<FName>( MeshAttribute::PolygonGroup::MaterialSlotName );
+			TPolygonGroupAttributeArray<FName>& PolygonGroupImportedMaterialSlotNames = MeshDescription->PolygonGroupAttributes().GetAttributes<FName>( MeshAttribute::PolygonGroup::ImportedMaterialSlotName );
+			TPolygonGroupAttributeArray<bool>& PolygonGroupCollision = MeshDescription->PolygonGroupAttributes().GetAttributes<bool>( MeshAttribute::PolygonGroup::EnableCollision );
+			TPolygonGroupAttributeArray<bool>& PolygonGroupCastShadow = MeshDescription->PolygonGroupAttributes().GetAttributes<bool>( MeshAttribute::PolygonGroup::CastShadow );
+
 			TArray<FStaticMaterial> MaterialToAdd;
-			for (const FPolygonGroupID& PolygonGroupID : MeshDescription->PolygonGroups().GetElementIDs())
+			for (const FPolygonGroupID PolygonGroupID : MeshDescription->PolygonGroups().GetElementIDs())
 			{
 				int32 MaterialIndex = PolygonGroupID.GetValue();
-				const FMeshPolygonGroup& MeshPolygonGroup = MeshDescription->GetPolygonGroup(PolygonGroupID);
-				FStaticMaterial StaticMaterial(Cast<UMaterialInterface>(MeshPolygonGroup.MaterialAsset.TryLoad()), MeshPolygonGroup.MaterialSlotName, MeshPolygonGroup.ImportedMaterialSlotName);
+				UMaterialInterface* Material = Cast<UMaterialInterface>(PolygonGroupMaterialAssets[PolygonGroupID]);
+				const FName& MaterialSlotName = PolygonGroupMaterialSlotNames[PolygonGroupID];
+				const FName& ImportedMaterialSlotName = PolygonGroupImportedMaterialSlotNames[PolygonGroupID];
+				FStaticMaterial StaticMaterial(Material, MaterialSlotName, ImportedMaterialSlotName);
 				int32 AddedIndex = (LODIndex > 0) ? MaterialToAdd.Add(StaticMaterial) : StaticMesh->StaticMaterials.Add(StaticMaterial);
 				check(MaterialIndex == AddedIndex);
 			}
diff --git a/Engine/Source/Editor/UnrealEd/Private/GeomFitUtils.cpp b/Engine/Source/Editor/UnrealEd/Private/GeomFitUtils.cpp
--- a/Engine/Source/Editor/UnrealEd/Private/GeomFitUtils.cpp
+++ b/Engine/Source/Editor/UnrealEd/Private/GeomFitUtils.cpp
@@ -21,6 +21,7 @@
 #include "PhysicsEngine/BodySetup.h"
 #include "Engine/StaticMesh.h"
 #include "MeshDescription.h"
+#include "MeshAttributes.h"
 
 #define LOCAL_EPS (0.01f)
 static void AddVertexIfNotPresent(TArray<FVector> &vertices, FVector &newVertex)
@@ -202,11 +203,12 @@ static void CalcBoundingBox(const FRawMesh& RawMesh, FVector& Center, FVector& E
 
 static void CalcBoundingBox(const UMeshDescription* MeshDescription, FVector& Center, FVector& Extents, FVector& LimitVec)
 {
+	const TVertexAttributeArray<FVector>& VertexPositions = MeshDescription->VertexAttributes().GetAttributes<FVector>(MeshAttribute::Vertex::Position);
+
 	FBox Box(ForceInit);
-	for (const FVertexID& VertexID : MeshDescription->Vertices().GetElementIDs())
+	for (const FVertexID VertexID : MeshDescription->Vertices().GetElementIDs())
 	{
-		const FMeshVertex& Vertex = MeshDescription->GetVertex(VertexID);
-		Box += Vertex.VertexPosition * LimitVec;
+		Box += VertexPositions[VertexID] * LimitVec;
 	}
 	Box.GetCenterAndExtents(Center, Extents);
 }
@@ -418,24 +420,25 @@ static void CalcBoundingSphere(const UMeshDescription* MeshDescription, FSphere&
 	FVector MinIx[3];
 	FVector MaxIx[3];
 
+	const TVertexAttributeArray<FVector>& VertexPositions = MeshDescription->VertexAttributes().GetAttributes<FVector>(MeshAttribute::Vertex::Position);
+
 	bool bFirstVertex = true;
-	for (const FVertexID& VertexID : MeshDescription->Vertices().GetElementIDs())
+	for (const FVertexID VertexID : MeshDescription->Vertices().GetElementIDs())
 	{
-		const FMeshVertex& Vertex = MeshDescription->GetVertex(VertexID);
-		FVector p = Vertex.VertexPosition * LimitVec;
+		FVector p = VertexPositions[VertexID] * LimitVec;
 		if (bFirstVertex)
 		{
 			// First, find AABB, remembering furthest points in each dir.
 			Box.Min = p;
 			Box.Max = Box.Min;
 
-			MinIx[0] = Vertex.VertexPosition;
-			MinIx[1] = Vertex.VertexPosition;
-			MinIx[2] = Vertex.VertexPosition;
+			MinIx[0] = VertexPositions[VertexID];
+			MinIx[1] = VertexPositions[VertexID];
+			MinIx[2] = VertexPositions[VertexID];
 
-			MaxIx[0] = Vertex.VertexPosition;
-			MaxIx[1] = Vertex.VertexPosition;
-			MaxIx[2] = Vertex.VertexPosition;
+			MaxIx[0] = VertexPositions[VertexID];
+			MaxIx[1] = VertexPositions[VertexID];
+			MaxIx[2] = VertexPositions[VertexID];
 			bFirstVertex = false;
 			continue;
 		}
@@ -444,36 +447,36 @@ static void CalcBoundingSphere(const UMeshDescription* MeshDescription, FSphere&
 		if (p.X < Box.Min.X)
 		{
 			Box.Min.X = p.X;
-			MinIx[0] = Vertex.VertexPosition;
+			MinIx[0] = VertexPositions[VertexID];
 		}
 		else if (p.X > Box.Max.X)
 		{
 			Box.Max.X = p.X;
-			MaxIx[0] = Vertex.VertexPosition;
+			MaxIx[0] = VertexPositions[VertexID];
 		}
 
 		// Y //
 		if (p.Y < Box.Min.Y)
 		{
 			Box.Min.Y = p.Y;
-			MinIx[1] = Vertex.VertexPosition;
+			MinIx[1] = VertexPositions[VertexID];
 		}
 		else if (p.Y > Box.Max.Y)
 		{
 			Box.Max.Y = p.Y;
-			MaxIx[1] = Vertex.VertexPosition;
+			MaxIx[1] = VertexPositions[VertexID];
 		}
 
 		// Z //
 		if (p.Z < Box.Min.Z)
 		{
 			Box.Min.Z = p.Z;
-			MinIx[2] = Vertex.VertexPosition;
+			MinIx[2] = VertexPositions[VertexID];
 		}
 		else if (p.Z > Box.Max.Z)
 		{
 			Box.Max.Z = p.Z;
-			MaxIx[2] = Vertex.VertexPosition;
+			MaxIx[2] = VertexPositions[VertexID];
 		}
 	}
 
@@ -501,11 +504,9 @@ static void CalcBoundingSphere(const UMeshDescription* MeshDescription, FSphere&
 	float r2 = FMath::Square(r);
 
 	// Now check each point lies within this sphere. If not - expand it a bit.
-	for (const FVertexID& VertexID : MeshDescription->Vertices().GetElementIDs())
+	for (const FVertexID VertexID : MeshDescription->Vertices().GetElementIDs())
 	{
-		const FMeshVertex& Vertex = MeshDescription->GetVertex(VertexID);
-
-		const FVector cToP = (Vertex.VertexPosition * LimitVec) - sphere.Center;
+		const FVector cToP = (VertexPositions[VertexID] * LimitVec) - sphere.Center;
 
 		const float pr2 = cToP.SizeSquared();
 
@@ -534,11 +535,11 @@ static void CalcBoundingSphere2(const UMeshDescription* MeshDescription, FSphere
 	sphere.Center = Center;
 	sphere.W = 0.0f;
 
-	for (const FVertexID& VertexID : MeshDescription->Vertices().GetElementIDs())
-	{
-		const FMeshVertex& Vertex = MeshDescription->GetVertex(VertexID);
+	const TVertexAttributeArray<FVector>& VertexPositions = MeshDescription->VertexAttributes().GetAttributes<FVector>(MeshAttribute::Vertex::Position);
 
-		float Dist = FVector::DistSquared(Vertex.VertexPosition * LimitVec, sphere.Center);
+	for (const FVertexID VertexID : MeshDescription->Vertices().GetElementIDs())
+	{
+		float Dist = FVector::DistSquared(VertexPositions[VertexID] * LimitVec, sphere.Center);
 		if (Dist > sphere.W)
 			sphere.W = Dist;
 	}
@@ -641,11 +642,12 @@ static void CalcBoundingSphyl(const UMeshDescription* MeshDescription, FSphere&
 	float r = Extents.GetMax();
 	float r2 = FMath::Square(r);
 	
+	const TVertexAttributeArray<FVector>& VertexPositions = MeshDescription->VertexAttributes().GetAttributes<FVector>(MeshAttribute::Vertex::Position);
+
 	// Now check each point lies within this the radius. If not - expand it a bit.
-	for (const FVertexID& VertexID : MeshDescription->Vertices().GetElementIDs())
+	for (const FVertexID VertexID : MeshDescription->Vertices().GetElementIDs())
 	{
-		const FMeshVertex& Vertex = MeshDescription->GetVertex(VertexID);
-		FVector cToP = (Vertex.VertexPosition * LimitVec) - sphere.Center;
+		FVector cToP = (VertexPositions[VertexID] * LimitVec) - sphere.Center;
 		cToP = rotation.UnrotateVector(cToP);
 
 		const float pr2 = cToP.SizeSquared2D();	// Ignore Z here...
@@ -664,10 +666,9 @@ static void CalcBoundingSphyl(const UMeshDescription* MeshDescription, FSphere&
 	float hl = FMath::Max(0.0f, Extent - r);
 
 	// Now check each point lies within the length. If not - expand it a bit.
-	for (const FVertexID& VertexID : MeshDescription->Vertices().GetElementIDs())
+	for (const FVertexID VertexID : MeshDescription->Vertices().GetElementIDs())
 	{
-		const FMeshVertex& Vertex = MeshDescription->GetVertex(VertexID);
-		FVector cToP = (Vertex.VertexPosition * LimitVec) - sphere.Center;
+		FVector cToP = (VertexPositions[VertexID] * LimitVec) - sphere.Center;
 		cToP = rotation.UnrotateVector(cToP);
 
 		// If this point is outside our current bounding sphyl's length
diff --git a/Engine/Source/Runtime/MeshBuilder/Private/LayoutUV.cpp b/Engine/Source/Runtime/MeshBuilder/Private/LayoutUV.cpp
--- a/Engine/Source/Runtime/MeshBuilder/Private/LayoutUV.cpp
+++ b/Engine/Source/Runtime/MeshBuilder/Private/LayoutUV.cpp
@@ -34,11 +34,13 @@ void FLayoutUV::FindCharts( const TMultiMap<int32,int32>& OverlappingCorners )
 	TranslatedMatches.SetNumUninitialized( NumIndexes );
 	TexCoords.SetNumUninitialized( NumIndexes );
 	int32 VertexInstanceIndex = 0;
-	for (const FVertexInstanceID& VertexInstanceID : MeshDescription->VertexInstances().GetElementIDs())
+
+	const TVertexInstanceAttributeArray<FVector2D>& VertexUVs = MeshDescription->VertexInstanceAttributes().GetAttributes<FVector2D>(MeshAttribute::VertexInstance::TextureCoordinate, SrcChannel);
+	for (const FVertexInstanceID VertexInstanceID : MeshDescription->VertexInstances().GetElementIDs())
 	{
 		check(VertexInstanceIndex == VertexInstanceID.GetValue());
 		TranslatedMatches[VertexInstanceIndex] = -1;
-		TexCoords[VertexInstanceIndex] = MeshDescription->GetVertexInstance(VertexInstanceID).VertexUVs[SrcChannel];
+		TexCoords[VertexInstanceIndex] = VertexUVs[VertexInstanceID];
 		VertexInstanceIndex++;
 	}
 
@@ -158,6 +160,8 @@ void FLayoutUV::FindCharts( const TMultiMap<int32,int32>& OverlappingCorners )
 
 	TMap< uint32, int32 > DisjointSetToChartMap;
 
+	const TVertexAttributeArray<FVector>& VertexPositions = MeshDescription->VertexAttributes().GetAttributes<FVector>(MeshAttribute::Vertex::Position);
+
 	// Build Charts
 	for( uint32 Tri = 0; Tri < NumTris; )
 	{
@@ -185,7 +189,7 @@ void FLayoutUV::FindCharts( const TMultiMap<int32,int32>& OverlappingCorners )
 				uint32 Index = 3 * SortedTris[ Tri ] + k;
 
 				FVertexInstanceID VertexInstanceID(Index);
-				Positions[k] = MeshDescription->GetVertex(MeshDescription->GetVertexInstance(VertexInstanceID).VertexID).VertexPosition;
+				Positions[k] = VertexPositions[MeshDescription->GetVertexInstanceVertex(VertexInstanceID)];
 				UVs[k] = TexCoords[ Index ];
 
 				Chart.MinUV.X = FMath::Min( Chart.MinUV.X, UVs[k].X );
@@ -1047,6 +1051,16 @@ void FLayoutUV::RasterizeChart( const FMeshChart& Chart, uint32 RectW, uint32 Re
 
 void FLayoutUV::CommitPackedUVs()
 {
+	// If current DstChannel is out of range of the number of UVs defined by the mesh description, change the index count accordingly
+	const uint32 NumUVs = MeshDescription->VertexInstanceAttributes().GetAttributeIndexCount<FVector2D>(MeshAttribute::VertexInstance::TextureCoordinate);
+	if (DstChannel >= NumUVs)
+	{
+		MeshDescription->VertexInstanceAttributes().SetAttributeIndexCount<FVector2D>(MeshAttribute::VertexInstance::TextureCoordinate, DstChannel + 1);
+		ensure(false);	// not expecting it to get here
+	}
+
+	TVertexInstanceAttributeArray<FVector2D>& VertexUVs = MeshDescription->VertexInstanceAttributes().GetAttributes<FVector2D>(MeshAttribute::VertexInstance::TextureCoordinate, DstChannel);
+
 	// Commit chart UVs
 	for( int32 i = 0; i < Charts.Num(); i++ )
 	{
@@ -1062,13 +1076,8 @@ void FLayoutUV::CommitPackedUVs()
 			{
 				uint32 Index = 3 * SortedTris[ Tri ] + k;
 				const FVector2D& UV = TexCoords[ Index ];
-				FVertexInstanceID VertexInstanceID(Index);
-				FMeshVertexInstance& VertexInstance = MeshDescription->GetVertexInstance(VertexInstanceID);
-				if (VertexInstance.VertexUVs.Num() <= (int32)DstChannel)
-				{
-					VertexInstance.VertexUVs.AddZeroed((1 + DstChannel) - VertexInstance.VertexUVs.Num());
-				}
-				VertexInstance.VertexUVs[ DstChannel ] = UV.X * Chart.PackingScaleU + UV.Y * Chart.PackingScaleV + Chart.PackingBias;
+				const FVertexInstanceID VertexInstanceID(Index);
+				VertexUVs[VertexInstanceID] = UV.X * Chart.PackingScaleU + UV.Y * Chart.PackingScaleV + Chart.PackingBias;
 			}
 		}
 	}
diff --git a/Engine/Source/Runtime/MeshBuilder/Private/LayoutUV.h b/Engine/Source/Runtime/MeshBuilder/Private/LayoutUV.h
--- a/Engine/Source/Runtime/MeshBuilder/Private/LayoutUV.h
+++ b/Engine/Source/Runtime/MeshBuilder/Private/LayoutUV.h
@@ -6,6 +6,7 @@
 #include "Allocator2D.h"
 #include "MeshDescription.h"
 #include "MeshDescriptionHelper.h"
+#include "MeshAttributes.h"
 
 #define NEW_UVS_ARE_SAME THRESH_POINTS_ARE_SAME
 #define LEGACY_UVS_ARE_SAME (1.0f / 1024.0f)
@@ -96,84 +97,62 @@ inline bool FLayoutUV::PositionsMatch( uint32 a, uint32 b ) const
 {
 	const FVertexInstanceID VertexInstanceIDA(a);
 	const FVertexInstanceID VertexInstanceIDB(b);
-	check(MeshDescription->VertexInstances().IsValid(VertexInstanceIDA));
-	check(MeshDescription->VertexInstances().IsValid(VertexInstanceIDB));
-	const FMeshVertexInstance& VertexInstanceA = MeshDescription->GetVertexInstance(VertexInstanceIDA);
-	const FMeshVertexInstance& VertexInstanceB = MeshDescription->GetVertexInstance(VertexInstanceIDB);
-	const FMeshVertex& VertexA = MeshDescription->GetVertex(VertexInstanceA.VertexID);
-	const FMeshVertex& VertexB = MeshDescription->GetVertex(VertexInstanceB.VertexID);
-	return VertexA.VertexPosition.Equals(VertexB.VertexPosition, THRESH_POINTS_ARE_SAME);
+	const FVertexID VertexIDA = MeshDescription->GetVertexInstanceVertex(VertexInstanceIDA);
+	const FVertexID VertexIDB = MeshDescription->GetVertexInstanceVertex(VertexInstanceIDB);
+
+	const TVertexAttributeArray<FVector>& VertexPositions = MeshDescription->VertexAttributes().GetAttributes<FVector>(MeshAttribute::Vertex::Position);
+	return VertexPositions[VertexIDA].Equals(VertexPositions[VertexIDB], THRESH_POINTS_ARE_SAME);
 }
 
 inline bool FLayoutUV::NormalsMatch( uint32 a, uint32 b ) const
 {
-	const FVertexInstanceID VertexInstanceIDA(a);
-	const FVertexInstanceID VertexInstanceIDB(b);
-	check(MeshDescription->VertexInstances().IsValid(VertexInstanceIDA));
-	check(MeshDescription->VertexInstances().IsValid(VertexInstanceIDB));
-	const FMeshVertexInstance& VertexInstanceA = MeshDescription->GetVertexInstance(VertexInstanceIDA);
-	const FMeshVertexInstance& VertexInstanceB = MeshDescription->GetVertexInstance(VertexInstanceIDB);
-
-	if (!VertexInstanceA.VertexUVs.IsValidIndex(SrcChannel))
+	// If current SrcChannel is out of range of the number of UVs defined by the mesh description, just return true
+	// @todo: hopefully remove this check entirely and just ensure that the mesh description matches the inputs
+	const uint32 NumUVs = MeshDescription->VertexInstanceAttributes().GetAttributeIndexCount<FVector2D>(MeshAttribute::VertexInstance::TextureCoordinate);
+	if (SrcChannel >= NumUVs)
 	{
+		ensure(false);	// not expecting it to get here
 		return true;
 	}
-	return VertexInstanceA.Normal.Equals(VertexInstanceB.Normal, THRESH_NORMALS_ARE_SAME);
-}
 
-inline bool FLayoutUV::UVsMatch( uint32 a, uint32 b ) const
-{
 	const FVertexInstanceID VertexInstanceIDA(a);
 	const FVertexInstanceID VertexInstanceIDB(b);
-	check(MeshDescription->VertexInstances().IsValid(VertexInstanceIDA));
-	check(MeshDescription->VertexInstances().IsValid(VertexInstanceIDB));
-	const FMeshVertexInstance& VertexInstanceA = MeshDescription->GetVertexInstance(VertexInstanceIDA);
-	const FMeshVertexInstance& VertexInstanceB = MeshDescription->GetVertexInstance(VertexInstanceIDB);
 
-	if (!VertexInstanceA.VertexUVs.IsValidIndex(SrcChannel))
+	const TVertexInstanceAttributeArray<FVector>& VertexNormals = MeshDescription->VertexInstanceAttributes().GetAttributes<FVector>(MeshAttribute::VertexInstance::Normal);
+	return VertexNormals[VertexInstanceIDA].Equals(VertexNormals[VertexInstanceIDB], THRESH_NORMALS_ARE_SAME);
+}
+
+inline bool FLayoutUV::UVsMatch( uint32 a, uint32 b ) const
+{
+	// If current SrcChannel is out of range of the number of UVs defined by the mesh description, just return true
+	const uint32 NumUVs = MeshDescription->VertexInstanceAttributes().GetAttributeIndexCount<FVector2D>(MeshAttribute::VertexInstance::TextureCoordinate);
+	if (SrcChannel >= NumUVs)
 	{
+		ensure(false);	// not expecting it to get here
 		return true;
 	}
-	return VertexInstanceA.VertexUVs[SrcChannel].Equals(VertexInstanceB.VertexUVs[SrcChannel], THRESH_UVS_ARE_SAME);
+
+	const FVertexInstanceID VertexInstanceIDA(a);
+	const FVertexInstanceID VertexInstanceIDB(b);
+
+	const TVertexInstanceAttributeArray<FVector2D>& VertexUVs = MeshDescription->VertexInstanceAttributes().GetAttributes<FVector2D>(MeshAttribute::VertexInstance::TextureCoordinate, SrcChannel);
+	return VertexUVs[VertexInstanceIDA].Equals(VertexUVs[VertexInstanceIDB], THRESH_UVS_ARE_SAME);
 }
 
 inline bool FLayoutUV::VertsMatch( uint32 a, uint32 b ) const
 {
-	const FVertexInstanceID VertexInstanceIDA(a);
-	const FVertexInstanceID VertexInstanceIDB(b);
-	check(MeshDescription->VertexInstances().IsValid(VertexInstanceIDA));
-	check(MeshDescription->VertexInstances().IsValid(VertexInstanceIDB));
-	const FMeshVertexInstance& VertexInstanceA = MeshDescription->GetVertexInstance(VertexInstanceIDA);
-	const FMeshVertexInstance& VertexInstanceB = MeshDescription->GetVertexInstance(VertexInstanceIDB);
-	bool IsUVsMatch = false;
-	if (!VertexInstanceA.VertexUVs.IsValidIndex(SrcChannel))
-	{
-		IsUVsMatch = true;
-	}
-	else
-	{
-		IsUVsMatch = VertexInstanceA.VertexUVs[SrcChannel].Equals(VertexInstanceB.VertexUVs[SrcChannel], THRESH_UVS_ARE_SAME);
-	}
-	if (!IsUVsMatch)
-	{
-		return false;
-	}
-	const FMeshVertex& VertexA = MeshDescription->GetVertex(VertexInstanceA.VertexID);
-	const FMeshVertex& VertexB = MeshDescription->GetVertex(VertexInstanceB.VertexID);
-	return VertexA.VertexPosition.Equals(VertexB.VertexPosition, THRESH_POINTS_ARE_SAME);
+	return PositionsMatch(a, b) && UVsMatch(a, b);
 }
 
 // Signed UV area
 inline float FLayoutUV::TriangleUVArea( uint32 Tri ) const
 {
+	const TVertexInstanceAttributeArray<FVector2D>& VertexUVs = MeshDescription->VertexInstanceAttributes().GetAttributes<FVector2D>(MeshAttribute::VertexInstance::TextureCoordinate, SrcChannel);
+
 	FVector2D UVs[3];
 	for( int k = 0; k < 3; k++ )
 	{
-		FVertexInstanceID VertexInstanceID((3*Tri) + k);
-		check(MeshDescription->VertexInstances().IsValid(VertexInstanceID));
-		const FMeshVertexInstance& VertexInstance = MeshDescription->GetVertexInstance(VertexInstanceID);
-		check(VertexInstance.VertexUVs.IsValidIndex(SrcChannel));
-		UVs[k] = VertexInstance.VertexUVs[ SrcChannel ];
+		UVs[k] = VertexUVs[FVertexInstanceID((3*Tri) + k)];
 	}
 
 	FVector2D EdgeUV1 = UVs[1] - UVs[0];
diff --git a/Engine/Source/Runtime/MeshBuilder/Private/MeshDescriptionHelper.cpp b/Engine/Source/Runtime/MeshBuilder/Private/MeshDescriptionHelper.cpp
--- a/Engine/Source/Runtime/MeshBuilder/Private/MeshDescriptionHelper.cpp
+++ b/Engine/Source/Runtime/MeshBuilder/Private/MeshDescriptionHelper.cpp
@@ -8,6 +8,7 @@
 #include "Engine/EngineTypes.h"
 #include "Engine/StaticMesh.h"
 #include "MeshDescription.h"
+#include "MeshAttributes.h"
 #include "mikktspace.h"
 #include "Serialization/MemoryWriter.h"
 #include "Serialization/MemoryReader.h"
@@ -55,6 +56,9 @@ UMeshDescription* FMeshDescriptionHelper::GetRenderMeshDescription(UObject* Owne
 	float ComparisonThreshold = BuildSettings->bRemoveDegenerates ? THRESH_POINTS_ARE_SAME : 0.0f;
 
 	FVertexInstanceArray& VertexInstanceArray = RenderMeshDescription->VertexInstances();
+	TVertexInstanceAttributeArray<FVector>& Normals = RenderMeshDescription->VertexInstanceAttributes().GetAttributes<FVector>( MeshAttribute::VertexInstance::Normal );
+	TVertexInstanceAttributeArray<FVector>& Tangents = RenderMeshDescription->VertexInstanceAttributes().GetAttributes<FVector>( MeshAttribute::VertexInstance::Tangent );
+	TVertexInstanceAttributeArray<float>& BinormalSigns = RenderMeshDescription->VertexInstanceAttributes().GetAttributes<float>( MeshAttribute::VertexInstance::BinormalSign );
 
 	// Find overlapping corners to accelerate adjacency.
 	FindOverlappingCorners(OverlappingCorners, RenderMeshDescription, ComparisonThreshold);
@@ -76,23 +80,22 @@ UMeshDescription* FMeshDescriptionHelper::GetRenderMeshDescription(UObject* Owne
 		bool bComputeTangentLegacy = !BuildSettings->bUseMikkTSpace && (BuildSettings->bRecomputeNormals || BuildSettings->bRecomputeTangents);
 		bool bHasAllNormals = true;
 		bool bHasAllTangents = true;
-		for (const FVertexInstanceID& VertexInstanceID : VertexInstanceArray.GetElementIDs())
+		for (const FVertexInstanceID VertexInstanceID : VertexInstanceArray.GetElementIDs())
 		{
-			FMeshVertexInstance& VertexInstance = VertexInstanceArray[VertexInstanceID];
 			// Dump normals and tangents if we are recomputing them.
 			if (BuildSettings->bRecomputeTangents)
 			{
 				//Dump the tangents
-				VertexInstance.BinormalSign = 0.0f;
-				VertexInstance.Tangent = FVector(0.0f);
+				BinormalSigns[VertexInstanceID] = 0.0f;
+				Tangents[VertexInstanceID] = FVector(0.0f);
 			}
 			if (BuildSettings->bRecomputeNormals)
 			{
 				//Dump the normals
-				VertexInstance.Normal = FVector(0.0f);
+				Normals[VertexInstanceID] = FVector(0.0f);
 			}
-			bHasAllNormals &= !VertexInstance.Normal.IsNearlyZero();
-			bHasAllTangents &= !VertexInstance.Tangent.IsNearlyZero();
+			bHasAllNormals &= !Normals[VertexInstanceID].IsNearlyZero();
+			bHasAllTangents &= !Tangents[VertexInstanceID].IsNearlyZero();
 		}
 
 
@@ -115,7 +118,7 @@ UMeshDescription* FMeshDescriptionHelper::GetRenderMeshDescription(UObject* Owne
 
 	if (BuildSettings->bGenerateLightmapUVs && VertexInstanceArray.Num() > 0)
 	{
-		if (!RenderMeshDescription->GetVertexInstance(FVertexInstanceID(0)).VertexUVs.IsValidIndex(BuildSettings->SrcLightmapIndex))
+		if (BuildSettings->SrcLightmapIndex >= RenderMeshDescription->VertexInstanceAttributes().GetAttributeIndexCount<FVector2D>(MeshAttribute::VertexInstance::TextureCoordinate))
 		{
 			BuildSettings->SrcLightmapIndex = 0;
 		}
@@ -186,10 +189,12 @@ void FMeshDescriptionHelper::FindOverlappingCorners(TMultiMap<int32, int32>& Ove
 	// Create a list of vertex Z/index pairs
 	TArray<FIndexAndZ> VertIndexAndZ;
 	VertIndexAndZ.Reserve(NumWedges);
+
+	const TVertexAttributeArray<FVector>& VertexPositions = MeshDescription->VertexAttributes().GetAttributes<FVector>(MeshAttribute::Vertex::Position);
+
 	for (const FVertexInstanceID VertexInstanceID : VertexInstanceArray.GetElementIDs())
 	{
-		const FMeshVertexInstance& VertexInstance = VertexInstanceArray[VertexInstanceID];
-		new(VertIndexAndZ)FIndexAndZ(VertexInstanceID.GetValue(), VertexArray[VertexInstance.VertexID].VertexPosition);
+		new(VertIndexAndZ)FIndexAndZ(VertexInstanceID.GetValue(), VertexPositions[MeshDescription->GetVertexInstanceVertex(VertexInstanceID)]);
 	}
 
 	// Sort the vertices by z value
@@ -223,6 +228,12 @@ void FMeshDescriptionHelper::FindOverlappingCorners(const UMeshDescription* Mesh
 
 void FMeshDescriptionHelper::CreatePolygonNTB(UMeshDescription* MeshDescription, float ComparisonThreshold)
 {
+	const TVertexAttributeArray<FVector>& VertexPositions = MeshDescription->VertexAttributes().GetAttributes<FVector>(MeshAttribute::Vertex::Position);
+	TVertexInstanceAttributeArray<FVector2D>& VertexUVs = MeshDescription->VertexInstanceAttributes().GetAttributes<FVector2D>(MeshAttribute::VertexInstance::TextureCoordinate, 0);
+	TPolygonAttributeArray<FVector>& PolygonNormals = MeshDescription->PolygonAttributes().GetAttributes<FVector>(MeshAttribute::Polygon::Normal);
+	TPolygonAttributeArray<FVector>& PolygonTangents = MeshDescription->PolygonAttributes().GetAttributes<FVector>(MeshAttribute::Polygon::Tangent);
+	TPolygonAttributeArray<FVector>& PolygonBinormals = MeshDescription->PolygonAttributes().GetAttributes<FVector>(MeshAttribute::Polygon::Binormal);
+
 	BuildStatisticManager::FBuildStatisticScope StatScope(TEXT("CreatePolygonNTB took"));
 	TArray<FPolygonID> DegeneratePolygons;
 
@@ -230,19 +241,18 @@ void FMeshDescriptionHelper::CreatePolygonNTB(UMeshDescription* MeshDescription,
 	FVertexArray& VertexArray = MeshDescription->Vertices();
 	FPolygonArray& PolygonArray = MeshDescription->Polygons();
 
-	for (const FPolygonID& PolygonID : MeshDescription->Polygons().GetElementIDs())
+	for (const FPolygonID PolygonID : MeshDescription->Polygons().GetElementIDs())
 	{
 		FVector TangentX(0.0f);
 		FVector TangentY(0.0f);
 		FVector TangentZ(0.0f);
 
-		FMeshPolygon& Polygon = PolygonArray[PolygonID];
-		if (!Polygon.PolygonNormal.IsNearlyZero())
+		if (!PolygonNormals[PolygonID].IsNearlyZero())
 		{
 			//By pass normal calculation if its already done
 			continue;
 		}
-		const TArray<FMeshTriangle>& MeshTriangles = Polygon.Triangles;
+		const TArray<FMeshTriangle>& MeshTriangles = MeshDescription->GetPolygonTriangles(PolygonID);
 #ifdef ENABLE_NTB_CHECK
 		//Assume triangle are build
 		check(MeshTriangles.Num() > 0);
@@ -254,14 +264,12 @@ void FMeshDescriptionHelper::CreatePolygonNTB(UMeshDescription* MeshDescription,
 
 		FVector P[3];
 		FVector2D UVs[3];
-		FMeshVertexInstance* MeshVertexInstance[3];
-		FMeshVertex* MeshVertex[3];
+
 		for (int32 i = 0; i < 3; ++i)
 		{
-			MeshVertexInstance[i] = &(VertexInstanceArray[MeshTriangle.GetVertexInstanceID(i)]);
-			UVs[i] = MeshVertexInstance[i]->VertexUVs[0];
-			MeshVertex[i] = &(VertexArray[MeshVertexInstance[i]->VertexID]);
-			P[i] = MeshVertex[i]->VertexPosition;
+			const FVertexInstanceID VertexInstanceID = MeshTriangle.GetVertexInstanceID(i);
+			UVs[i] = VertexUVs[VertexInstanceID];
+			P[i] = VertexPositions[MeshDescription->GetVertexInstanceVertex(VertexInstanceID)];
 		}
 
 		const FVector Normal = ((P[1] - P[2]) ^ (P[0] - P[2])).GetSafeNormal(ComparisonThreshold);
@@ -294,18 +302,19 @@ void FMeshDescriptionHelper::CreatePolygonNTB(UMeshDescription* MeshDescription,
 		{
 			DegeneratePolygons.Add(PolygonID);
 		}
-		Polygon.PolygonTangent = TangentX;
-		Polygon.PolygonBinormal = TangentY;
-		Polygon.PolygonNormal = TangentZ;
+
+		PolygonTangents[PolygonID] = TangentX;
+		PolygonBinormals[PolygonID] = TangentY;
+		PolygonNormals[PolygonID] = TangentZ;
 	}
 
 	//Delete the degenerated polygons. The array is fill only if the remove degenerated option is turn on.
 	if (DegeneratePolygons.Num() > 0)
 	{
-		TSet<FEdgeID> OrphanedEdges;
-		TSet<FVertexInstanceID> OrphanedVertexInstances;
-		TSet<FPolygonGroupID> OrphanedPolygonGroups;
-		TSet<FVertexID> OrphanedVertices;
+		TArray<FEdgeID> OrphanedEdges;
+		TArray<FVertexInstanceID> OrphanedVertexInstances;
+		TArray<FPolygonGroupID> OrphanedPolygonGroups;
+		TArray<FVertexID> OrphanedVertices;
 		for (FPolygonID& PolygonID : DegeneratePolygons)
 		{
 			MeshDescription->DeletePolygon(PolygonID, &OrphanedEdges, &OrphanedVertexInstances, &OrphanedPolygonGroups);
@@ -341,19 +350,17 @@ struct FVertexInfo
 		VertexInstanceID = FVertexInstanceID::Invalid;
 		UVs = FVector2D(0.0f, 0.0f);
 		EdgeIDs.Reserve(2);//Most of the time a edge has two triangles
-		Polygon = nullptr;
 	}
 
 	FPolygonID PolygonID;
-	FMeshPolygon* Polygon;
 	FVertexInstanceID VertexInstanceID;
 	FVector2D UVs;
 	TArray<FEdgeID> EdgeIDs;
-	TMap<FEdgeID, FMeshEdge*> Edges;
 };
 
 void RecursiveFillPolygonGroup(
-	  TMap<FPolygonID, FVertexInfo>& VertexInfoMap
+	  UMeshDescription* MeshDescription
+	, TMap<FPolygonID, FVertexInfo>& VertexInfoMap
 	, TArray<FPolygonID>& CurrentGroup
 	, FVertexInfo& CurrentVertexInfo
 	, TArray<FPolygonID>& ConsumedPolygon)
@@ -361,15 +368,16 @@ void RecursiveFillPolygonGroup(
 	CurrentGroup.AddUnique(CurrentVertexInfo.PolygonID);
 	ConsumedPolygon.AddUnique(CurrentVertexInfo.PolygonID);
 
-	for (FEdgeID &EdgeID : CurrentVertexInfo.EdgeIDs)
+	const TEdgeAttributeArray<bool>& EdgeHardnesses = MeshDescription->EdgeAttributes().GetAttributes<bool>(MeshAttribute::Edge::IsHard);
+
+	for (const FEdgeID EdgeID : CurrentVertexInfo.EdgeIDs)
 	{
-		FMeshEdge& Edge = *(CurrentVertexInfo.Edges[EdgeID]);
-		if (Edge.bIsHardEdge)
+		if (EdgeHardnesses[EdgeID])
 		{
 			//End of the group
 			continue;
 		}
-		for (const FPolygonID& PolygonID : Edge.ConnectedPolygons)
+		for (const FPolygonID PolygonID : MeshDescription->GetEdgeConnectedPolygons(EdgeID))
 		{
 			if (PolygonID == CurrentVertexInfo.PolygonID)
 			{
@@ -380,7 +388,7 @@ void RecursiveFillPolygonGroup(
 			//Do not repeat polygons
 			if (!ConsumedPolygon.Contains(OtherVertexInfo.PolygonID))
 			{
-				RecursiveFillPolygonGroup(VertexInfoMap, CurrentGroup, OtherVertexInfo, ConsumedPolygon);
+				RecursiveFillPolygonGroup(MeshDescription, VertexInfoMap, CurrentGroup, OtherVertexInfo, ConsumedPolygon);
 			}
 		}
 	}
@@ -402,48 +410,45 @@ void FMeshDescriptionHelper::CreateNormals(UMeshDescription* MeshDescription, FM
 	// First normal(A to F) = Normalize(A+B+C+D+E+F)
 	// Second normal(F to A) = Normalize(F+G+H+A)
 	// We found the connected edge using the triangle that share edges
+
+	// @todo: provide an option to weight each contributing polygon normal according to the size of
+	// the angle it makes with the vertex being calculated. This means that triangulated faces whose
+	// internal edge meets the vertex doesn't get undue extra weight.
 	
-	FVertexInstanceArray& VertexInstanceArray = MeshDescription->VertexInstances();
-	FVertexArray& VertexArray = MeshDescription->Vertices();
-	FPolygonArray& PolygonArray = MeshDescription->Polygons();
-	FEdgeArray& EdgeArray = MeshDescription->Edges();
+	const TVertexInstanceAttributeArray<FVector2D>& VertexUVs = MeshDescription->VertexInstanceAttributes().GetAttributes<FVector2D>(MeshAttribute::VertexInstance::TextureCoordinate, 0);
+	TVertexInstanceAttributeArray<FVector>& VertexNormals = MeshDescription->VertexInstanceAttributes().GetAttributes<FVector>(MeshAttribute::VertexInstance::Normal, 0);
+	TVertexInstanceAttributeArray<FVector>& VertexTangents = MeshDescription->VertexInstanceAttributes().GetAttributes<FVector>(MeshAttribute::VertexInstance::Tangent, 0);
+	TVertexInstanceAttributeArray<float>& VertexBinormalSigns = MeshDescription->VertexInstanceAttributes().GetAttributes<float>(MeshAttribute::VertexInstance::BinormalSign, 0);
+
+	TPolygonAttributeArray<FVector>& PolygonNormals = MeshDescription->PolygonAttributes().GetAttributes<FVector>(MeshAttribute::Polygon::Normal);
+	TPolygonAttributeArray<FVector>& PolygonTangents = MeshDescription->PolygonAttributes().GetAttributes<FVector>(MeshAttribute::Polygon::Tangent);
+	TPolygonAttributeArray<FVector>& PolygonBinormals = MeshDescription->PolygonAttributes().GetAttributes<FVector>(MeshAttribute::Polygon::Binormal);
 
 	TMap<FPolygonID, FVertexInfo> VertexInfoMap;
-	TMap<FVertexInstanceID, FMeshVertexInstance*> VertexInstanceMap;
 	VertexInfoMap.Reserve(20);
 	//Iterate all vertex to compute normals for all vertex instance
-	for (const FVertexID& VertexID : VertexArray.GetElementIDs())
+	for (const FVertexID VertexID : MeshDescription->Vertices().GetElementIDs())
 	{
-		FMeshVertex& Vertex = VertexArray[VertexID];
-
 		VertexInfoMap.Reset();
 		
 		bool bPointHasAllTangents = true;
 		//Fill the VertexInfoMap
-		for (const FEdgeID& EdgeID : Vertex.ConnectedEdgeIDs)
+		for (const FEdgeID EdgeID : MeshDescription->GetVertexConnectedEdges(VertexID))
 		{
-			FMeshEdge& Edge = EdgeArray[EdgeID];
-			for (const FPolygonID& PolygonID : Edge.ConnectedPolygons)
+			for (const FPolygonID PolygonID : MeshDescription->GetEdgeConnectedPolygons(EdgeID))
 			{
 				FVertexInfo& VertexInfo = VertexInfoMap.FindOrAdd(PolygonID);
 				int32 EdgeIndex = VertexInfo.EdgeIDs.AddUnique(EdgeID);
-				VertexInfo.Edges.Add(EdgeID, &Edge);
-
 				if (VertexInfo.PolygonID == FPolygonID::Invalid)
 				{
-					FMeshPolygon& Polygon = PolygonArray[PolygonID];
-
 					VertexInfo.PolygonID = PolygonID;
-					VertexInfo.Polygon = &Polygon;
-					for (FVertexInstanceID& VertexInstanceID : Polygon.PerimeterContour.VertexInstanceIDs)
+					for (const FVertexInstanceID VertexInstanceID : MeshDescription->GetPolygonPerimeterVertexInstances(PolygonID))
 					{
-						FMeshVertexInstance& VertexInstance = VertexInstanceArray[VertexInstanceID];
-						if (VertexInstance.VertexID == VertexID)
+						if (MeshDescription->GetVertexInstanceVertex(VertexInstanceID) == VertexID)
 						{
 							VertexInfo.VertexInstanceID = VertexInstanceID;
-							VertexInfo.UVs = VertexInstance.VertexUVs[0];
-							VertexInstanceMap.Add(VertexInstanceID, &VertexInstance);
-							bPointHasAllTangents &= !VertexInstance.Normal.IsNearlyZero() && !VertexInstance.Tangent.IsNearlyZero();
+							VertexInfo.UVs = VertexUVs[VertexInstanceID];
+							bPointHasAllTangents &= !VertexNormals[VertexInstanceID].IsNearlyZero() && !VertexTangents[VertexInstanceID].IsNearlyZero();
 							break;
 						}
 					}
@@ -457,7 +462,7 @@ void FMeshDescriptionHelper::CreateNormals(UMeshDescription* MeshDescription, FM
 		}
 
 		//Make sure we consume all our vertex instance
-		check(VertexInfoMap.Num() == Vertex.VertexInstanceIDs.Num());
+		check(VertexInfoMap.Num() == MeshDescription->GetVertexVertexInstances(VertexID).Num());
 
 		//Build all group by recursively traverse all polygon connected to the vertex
 		TArray<TArray<FPolygonID>> Groups;
@@ -471,7 +476,7 @@ void FMeshDescriptionHelper::CreateNormals(UMeshDescription* MeshDescription, FM
 
 			int32 CurrentGroupIndex = Groups.AddZeroed();
 			FVertexInfo& VertexInfo = VertexInfoMap.FindOrAdd(Kvp.Key);
-			RecursiveFillPolygonGroup(VertexInfoMap, Groups[CurrentGroupIndex], VertexInfo, ConsumedPolygon);
+			RecursiveFillPolygonGroup(MeshDescription, VertexInfoMap, Groups[CurrentGroupIndex], VertexInfo, ConsumedPolygon);
 		}
 
 		//Smooth every connected group
@@ -484,23 +489,22 @@ void FMeshDescriptionHelper::CreateNormals(UMeshDescription* MeshDescription, FM
 
 			TArray<FVertexInstanceID> VertexInstanceInGroup;
 			FVector GroupNormal(0.0f);
-			for(const FPolygonID& PolygonID : Group)
+			for(const FPolygonID PolygonID : Group)
 			{
 #ifdef ENABLE_NTB_CHECK
 				check(!ConsumedPolygon.Contains(PolygonID));
 #endif
 				ConsumedPolygon.Add(PolygonID);
 				VertexInstanceInGroup.Add(VertexInfoMap[PolygonID].VertexInstanceID);
-				FMeshPolygon& Polygon = *(VertexInfoMap[PolygonID].Polygon);
-				GroupNormal += Polygon.PolygonNormal;
+				GroupNormal += PolygonNormals[PolygonID];
 				if (bComputeTangent)
 				{
-					FVector2D& UVs = VertexInfoMap[PolygonID].UVs;
+					const FVector2D UVs = VertexInfoMap[PolygonID].UVs;
 					bool CreateGroup = (!GroupTangent.Contains(UVs));
 					FVector& GroupTangentValue = GroupTangent.FindOrAdd(UVs);
 					FVector& GroupBiNormalValue = GroupBiNormal.FindOrAdd(UVs);
-					GroupTangentValue = CreateGroup ? Polygon.PolygonTangent : GroupTangentValue + Polygon.PolygonTangent;
-					GroupBiNormalValue = CreateGroup ? Polygon.PolygonBinormal : GroupBiNormalValue + Polygon.PolygonBinormal;
+					GroupTangentValue = CreateGroup ? PolygonTangents[PolygonID] : GroupTangentValue + PolygonTangents[PolygonID];
+					GroupBiNormalValue = CreateGroup ? PolygonBinormals[PolygonID] : GroupBiNormalValue + PolygonBinormals[PolygonID];
 				}
 			}
 
@@ -521,17 +525,16 @@ void FMeshDescriptionHelper::CreateNormals(UMeshDescription* MeshDescription, FM
 				}
 			}
 			//Apply the average NTB on all Vertex instance
-			for (FVertexInstanceID& VertexInstanceID : VertexInstanceInGroup)
+			for (const FVertexInstanceID VertexInstanceID : VertexInstanceInGroup)
 			{
-				FMeshVertexInstance& VertexInstance = *(VertexInstanceMap[VertexInstanceID]);
-				FVector2D& VertexUV = VertexInstance.VertexUVs[0];
+				const FVector2D VertexUV = VertexUVs[VertexInstanceID];
 			
-				if (VertexInstance.Normal.IsNearlyZero(SMALL_NUMBER))
+				if (VertexNormals[VertexInstanceID].IsNearlyZero(SMALL_NUMBER))
 				{
-					VertexInstance.Normal = GroupNormal;
+					VertexNormals[VertexInstanceID] = GroupNormal;
 				}
 #ifdef ENABLE_NTB_CHECK
-				check(!VertexInstance.Normal.IsNearlyZero(SMALL_NUMBER));
+				check(!VertexNormals[VertexInstanceID].IsNearlyZero(SMALL_NUMBER));
 #endif
 				if (bComputeTangent)
 				{
@@ -539,17 +542,17 @@ void FMeshDescriptionHelper::CreateNormals(UMeshDescription* MeshDescription, FM
 					FVector GroupTangentValue = GroupTangent[VertexUV];
 					FVector GroupBiNormalValue = GroupBiNormal[VertexUV];
 
-					if (!VertexInstance.Tangent.IsNearlyZero(SMALL_NUMBER))
+					if (!VertexTangents[VertexInstanceID].IsNearlyZero(SMALL_NUMBER))
 					{
-						GroupTangentValue = VertexInstance.Tangent;
+						GroupTangentValue = VertexTangents[VertexInstanceID];
 					}
 #ifdef ENABLE_NTB_CHECK
 					check(!GroupTangentValue.IsNearlyZero(SMALL_NUMBER))
 #endif
 					FVector BiNormal(0.0f);
-					if (!VertexInstance.Normal.IsNearlyZero(SMALL_NUMBER) && !VertexInstance.Tangent.IsNearlyZero(SMALL_NUMBER))
+					if (!VertexNormals[VertexInstanceID].IsNearlyZero(SMALL_NUMBER) && !VertexTangents[VertexInstanceID].IsNearlyZero(SMALL_NUMBER))
 					{
-						BiNormal = FVector::CrossProduct(VertexInstance.Normal, VertexInstance.Tangent).GetSafeNormal() * VertexInstance.BinormalSign;
+						BiNormal = FVector::CrossProduct(VertexNormals[VertexInstanceID], VertexTangents[VertexInstanceID]).GetSafeNormal() * VertexBinormalSigns[VertexInstanceID];
 					}
 					if (!BiNormal.IsNearlyZero(SMALL_NUMBER))
 					{
@@ -562,19 +565,20 @@ void FMeshDescriptionHelper::CreateNormals(UMeshDescription* MeshDescription, FM
 					GroupBiNormalValue -= GroupTangentValue * (GroupTangentValue | GroupBiNormalValue);
 					GroupBiNormalValue.Normalize();
 					
-					GroupTangentValue -= VertexInstance.Normal * (VertexInstance.Normal | GroupTangentValue);
+					GroupTangentValue -= VertexNormals[VertexInstanceID] * (VertexNormals[VertexInstanceID] | GroupTangentValue);
 					GroupTangentValue.Normalize();
 					
-					GroupBiNormalValue -= VertexInstance.Normal * (VertexInstance.Normal | GroupBiNormalValue);
+					GroupBiNormalValue -= VertexNormals[VertexInstanceID] * (VertexNormals[VertexInstanceID] | GroupBiNormalValue);
 					GroupBiNormalValue.Normalize();
 #ifdef ENABLE_NTB_CHECK
 					check(!GroupTangentValue.IsNearlyZero(SMALL_NUMBER));
 					check(!GroupBiNormalValue.IsNearlyZero(SMALL_NUMBER));
 #endif
 					//Set the value
-					VertexInstance.Tangent = GroupTangentValue;
+					VertexTangents[VertexInstanceID] = GroupTangentValue;
 					//If the BiNormal is zero set the sign to 1.0f
-					VertexInstance.BinormalSign = GetBasisDeterminantSign(GroupTangentValue, GroupBiNormalValue, VertexInstance.Normal);
+					VertexBinormalSigns[VertexInstanceID] = GetBasisDeterminantSign(GroupTangentValue, GroupBiNormalValue, VertexNormals[VertexInstanceID]);
+
 				}
 			}
 		}
@@ -620,7 +624,9 @@ namespace MeshDescriptionMikktSpaceInterface
 	void MikkGetPosition(const SMikkTSpaceContext* Context, float Position[3], const int FaceIdx, const int VertIdx)
 	{
 		UMeshDescription *MeshDescription = (UMeshDescription*)(Context->m_pUserData);
-		const FVector &VertexPosition = MeshDescription->Vertices()[MeshDescription->VertexInstances()[FVertexInstanceID(FaceIdx * 3 + VertIdx)].VertexID].VertexPosition;
+		const FVertexInstanceID VertexInstanceID(FaceIdx * 3 + VertIdx);
+		const FVertexID VertexID = MeshDescription->GetVertexInstanceVertex(VertexInstanceID);
+		const FVector& VertexPosition = MeshDescription->VertexAttributes().GetAttribute<FVector>(VertexID, MeshAttribute::Vertex::Position);
 		Position[0] = VertexPosition.X;
 		Position[1] = VertexPosition.Y;
 		Position[2] = VertexPosition.Z;
@@ -629,31 +635,27 @@ namespace MeshDescriptionMikktSpaceInterface
 	void MikkGetNormal(const SMikkTSpaceContext* Context, float Normal[3], const int FaceIdx, const int VertIdx)
 	{
 		UMeshDescription *MeshDescription = (UMeshDescription*)(Context->m_pUserData);
-		const FMeshVertexInstance& VertexInstance = MeshDescription->VertexInstances()[FVertexInstanceID(FaceIdx * 3 + VertIdx)];
-		const FVector &VertexNormal = VertexInstance.Normal;
-		for (int32 i = 0; i < 3; ++i)
-		{
-			Normal[i] = VertexNormal[i];
-		}
+		const FVertexInstanceID VertexInstanceID(FaceIdx * 3 + VertIdx);
+		const FVector& VertexNormal = MeshDescription->VertexInstanceAttributes().GetAttribute<FVector>(VertexInstanceID, MeshAttribute::VertexInstance::Normal);
+		Normal[0] = VertexNormal.X;
+		Normal[1] = VertexNormal.Y;
+		Normal[2] = VertexNormal.Z;
 	}
 
 	void MikkSetTSpaceBasic(const SMikkTSpaceContext* Context, const float Tangent[3], const float BitangentSign, const int FaceIdx, const int VertIdx)
 	{
 		UMeshDescription *MeshDescription = (UMeshDescription*)(Context->m_pUserData);
-		FMeshVertexInstance& VertexInstance = MeshDescription->VertexInstances()[FVertexInstanceID(FaceIdx * 3 + VertIdx)];
-		//Copy the tangents
-		for (int32 i = 0; i < 3; ++i)
-		{
-			VertexInstance.Tangent[i] = Tangent[i];
-		}
-		//Copy the bi normal sign
-		VertexInstance.BinormalSign = -BitangentSign;
+		const FVertexInstanceID VertexInstanceID(FaceIdx * 3 + VertIdx);
+		const FVector VertexTangent(Tangent[0], Tangent[1], Tangent[2]);
+		MeshDescription->VertexInstanceAttributes().SetAttribute<FVector>(VertexInstanceID, MeshAttribute::VertexInstance::Tangent, 0, VertexTangent);
+		MeshDescription->VertexInstanceAttributes().SetAttribute<float>(VertexInstanceID, MeshAttribute::VertexInstance::BinormalSign, 0, -BitangentSign);
 	}
 
 	void MikkGetTexCoord(const SMikkTSpaceContext* Context, float UV[2], const int FaceIdx, const int VertIdx)
 	{
 		UMeshDescription *MeshDescription = (UMeshDescription*)(Context->m_pUserData);
-		const FVector2D &TexCoord = MeshDescription->VertexInstances()[FVertexInstanceID(FaceIdx * 3 + VertIdx)].VertexUVs[0];
+		const FVertexInstanceID VertexInstanceID(FaceIdx * 3 + VertIdx);
+		const FVector2D& TexCoord = MeshDescription->VertexInstanceAttributes().GetAttribute<FVector2D>(VertexInstanceID, MeshAttribute::VertexInstance::TextureCoordinate, 0);
 		UV[0] = TexCoord.X;
 		UV[1] = TexCoord.Y;
 	}
diff --git a/Engine/Source/Runtime/MeshBuilder/Private/MeshDescriptionHelper.h b/Engine/Source/Runtime/MeshBuilder/Private/MeshDescriptionHelper.h
--- a/Engine/Source/Runtime/MeshBuilder/Private/MeshDescriptionHelper.h
+++ b/Engine/Source/Runtime/MeshBuilder/Private/MeshDescriptionHelper.h
@@ -58,8 +58,6 @@ private:
 	//////////////////////////////////////////////////////////////////////////
 	//PRIVATE function declarations
 
-	FORCEINLINE const FVector& GetVertexPositionFromVertexInstance(UMeshDescription* MeshDescription, const FVertexInstanceID& VertexInstanceID) const;
-	FORCEINLINE FVector2D& GetVertexInstanceUV(UMeshDescription* MeshDescription, const FVertexInstanceID& VertexInstanceID, int32 UVLayer) const;
 	void CreateNormals(UMeshDescription* MeshDescription, ETangentOptions TangentOptions, bool bComputeTangent);
 	void CreateMikktTangents(UMeshDescription* MeshDescription, ETangentOptions TangentOptions);
 	void CreatePolygonNTB(UMeshDescription* MeshDescription, float ComparisonThreshold);
diff --git a/Engine/Source/Runtime/MeshBuilder/Private/StaticMeshBuilder.cpp b/Engine/Source/Runtime/MeshBuilder/Private/StaticMeshBuilder.cpp
--- a/Engine/Source/Runtime/MeshBuilder/Private/StaticMeshBuilder.cpp
+++ b/Engine/Source/Runtime/MeshBuilder/Private/StaticMeshBuilder.cpp
@@ -5,6 +5,7 @@
 #include "StaticMeshResources.h"
 #include "PhysicsEngine/BodySetup.h"
 #include "MeshDescription.h"
+#include "MeshAttributes.h"
 #include "MeshDescriptionHelper.h"
 #include "BuildOptimizationHelper.h"
 #include "Components.h"
@@ -187,10 +188,9 @@ void FStaticMeshBuilder::OnBuildRenderMeshFinish(UStaticMesh* StaticMesh, const
 
 bool IsOrphanedVertex(const UMeshDescription *MeshDescription, const FVertexID VertexID)
 {
-	const FVertexInstanceArray& VertexInstances = MeshDescription->VertexInstances();
-	for (const FVertexInstanceID VertexInstanceID : MeshDescription->Vertices()[VertexID].VertexInstanceIDs)
+	for (const FVertexInstanceID VertexInstanceID : MeshDescription->GetVertexVertexInstances(VertexID))
 	{
-		if (VertexInstances[VertexInstanceID].ConnectedPolygons.Num() > 0)
+		if (MeshDescription->GetVertexInstanceConnectedPolygons(VertexInstanceID).Num() > 0)
 		{
 			return false;
 		}
@@ -219,11 +219,13 @@ void UpdateBounds(UStaticMesh* StaticMesh)
 	const FVertexArray& Vertices = MeshDescription->Vertices();
 	const FVertexInstanceArray& VertexInstances = MeshDescription->VertexInstances();
 
+	const TVertexAttributeArray<FVector>& VertexPositions = MeshDescription->VertexAttributes().GetAttributes<FVector>(MeshAttribute::Vertex::Position);
+
 	for (const FVertexID VertexID : Vertices.GetElementIDs())
 	{
 		if (!IsOrphanedVertex(MeshDescription, VertexID))
 		{
-			BoundingBox += Vertices[VertexID].VertexPosition;
+			BoundingBox += VertexPositions[VertexID];
 		}
 	}
 
@@ -236,8 +238,7 @@ void UpdateBounds(UStaticMesh* StaticMesh)
 	{
 		if (!IsOrphanedVertex(MeshDescription, VertexID))
 		{
-			const FVector VertexPosition = Vertices[VertexID].VertexPosition;
-			BoundingBoxAndSphere.SphereRadius = FMath::Max((VertexPosition - BoundingBoxAndSphere.Origin).Size(), BoundingBoxAndSphere.SphereRadius);
+			BoundingBoxAndSphere.SphereRadius = FMath::Max((VertexPositions[VertexID] - BoundingBoxAndSphere.Origin).Size(), BoundingBoxAndSphere.SphereRadius);
 		}
 	}
 
@@ -362,14 +363,25 @@ void BuildVertexBuffer(
 	}
 	TArray<int32> DupVerts;
 
-	uint32 NumTextureCoord = (uint32)(VertexInstances.Num() > 0 ? VertexInstances[FVertexInstanceID(0)].VertexUVs.Num() : 1);
+	const uint32 NumTextureCoord = MeshDescription->VertexInstanceAttributes().GetAttributeIndexCount<FVector2D>(MeshAttribute::VertexInstance::TextureCoordinate);
+
+	const TPolygonGroupAttributeArray<FName>& PolygonGroupImportedMaterialSlotNames = MeshDescription->PolygonGroupAttributes().GetAttributes<FName>(MeshAttribute::PolygonGroup::ImportedMaterialSlotName);
+	const TPolygonGroupAttributeArray<UObject*>& PolygonGroupMaterialAssets = MeshDescription->PolygonGroupAttributes().GetAttributes<UObject*>(MeshAttribute::PolygonGroup::MaterialAsset);
+	const TPolygonGroupAttributeArray<bool>& PolygonGroupCollision = MeshDescription->PolygonGroupAttributes().GetAttributes<bool>(MeshAttribute::PolygonGroup::EnableCollision);
+	const TPolygonGroupAttributeArray<bool>& PolygonGroupCastShadow = MeshDescription->PolygonGroupAttributes().GetAttributes<bool>(MeshAttribute::PolygonGroup::CastShadow);
+
+	const TVertexAttributeArray<FVector>& VertexPositions = MeshDescription->VertexAttributes().GetAttributes<FVector>( MeshAttribute::Vertex::Position );
+	const TVertexInstanceAttributeArray<FVector>& VertexInstanceNormals = MeshDescription->VertexInstanceAttributes().GetAttributes<FVector>( MeshAttribute::VertexInstance::Normal );
+	const TVertexInstanceAttributeArray<FVector>& VertexInstanceTangents = MeshDescription->VertexInstanceAttributes().GetAttributes<FVector>( MeshAttribute::VertexInstance::Tangent );
+	const TVertexInstanceAttributeArray<float>& VertexInstanceBinormalSigns = MeshDescription->VertexInstanceAttributes().GetAttributes<float>( MeshAttribute::VertexInstance::BinormalSign );
+	const TVertexInstanceAttributeArray<FVector4>& VertexInstanceColors = MeshDescription->VertexInstanceAttributes().GetAttributes<FVector4>( MeshAttribute::VertexInstance::Color );
+	const TVertexInstanceAttributeIndicesArray<FVector2D>& VertexInstanceUVs = MeshDescription->VertexInstanceAttributes().GetAttributesSet<FVector2D>( MeshAttribute::VertexInstance::TextureCoordinate );
 
 	// Set up index buffer
-	for (const FPolygonGroupID& PolygonGroupID : PolygonGroupArray.GetElementIDs())
+	for (const FPolygonGroupID PolygonGroupID : MeshDescription->PolygonGroups().GetElementIDs())
 	{
-		const FMeshPolygonGroup& PolygonGroup = PolygonGroupArray[PolygonGroupID];
-		const TArray<FPolygonID> &Polygons = PolygonGroup.Polygons;
-		
+		const TArray<FPolygonID>& Polygons = MeshDescription->GetPolygonGroupPolygons(PolygonGroupID);
+
 		TArray<uint32>& SectionIndices = OutPerSectionIndices[PolygonGroupID.GetValue()];
 
 		// Create new rendering section
@@ -377,15 +389,18 @@ void BuildVertexBuffer(
 		FStaticMeshSection& StaticMeshSection = StaticMeshLOD.Sections.Last();
 
 		StaticMeshSection.FirstIndex = IndexBuffer.Num();
+		// @todo: is currently assuming all polygons are triangles.
+		// At very least, this should be summing the sizes of all PolygonTriangles arrays.
+		// MeshDescription->GetPolygonTriangles() should always return a valid array of at least one triangle if the mesh description is valid.
 		StaticMeshSection.NumTriangles = Polygons.Num();
 
-		const int32 MaterialIndex = StaticMesh->GetMaterialIndex(PolygonGroup.ImportedMaterialSlotName);
+		const int32 MaterialIndex = StaticMesh->GetMaterialIndex(PolygonGroupImportedMaterialSlotNames[PolygonGroupID]);
 		check(MaterialIndex != INDEX_NONE);
-		check(StaticMesh->StaticMaterials[MaterialIndex].MaterialInterface->GetPathName() == PolygonGroup.MaterialAsset.ToString());
+		check(StaticMesh->StaticMaterials[MaterialIndex].MaterialInterface == Cast<UMaterialInterface>(PolygonGroupMaterialAssets[PolygonGroupID]));
 		StaticMeshSection.MaterialIndex = MaterialIndex;
-		StaticMeshSection.bEnableCollision = PolygonGroup.bEnableCollision;
-		StaticMeshSection.bCastShadow = PolygonGroup.bCastShadow;
-		
+		StaticMeshSection.bEnableCollision = PolygonGroupCollision[PolygonGroupID];
+		StaticMeshSection.bCastShadow = PolygonGroupCastShadow[PolygonGroupID];
+
 		if (LodIndex > 0)
 		{
 			//Set the overwrite section info map
@@ -396,39 +411,55 @@ void BuildVertexBuffer(
 			StaticMesh->SectionInfoMap.Set(LodIndex, SectionIndex, SectionInfo);
 		}
 
-		for (const FPolygonID& PolygonID : Polygons)
+		for (const FPolygonID PolygonID : Polygons)
 		{
-			const FMeshPolygon& Polygon = PolygonArray[PolygonID];
-			int32 ReserveSize = IndexBuffer.Num() + Polygon.Triangles.Num() * 3;
+			const TArray<FMeshTriangle>& PolygonTriangles = MeshDescription->GetPolygonTriangles(PolygonID);
+			int32 ReserveSize = IndexBuffer.Num() + PolygonTriangles.Num() * 3;
 			IndexBuffer.Reserve(ReserveSize);
 			OutWedgeMap.Reserve(ReserveSize);
 			uint32 MinIndex = TNumericLimits< uint32 >::Max();
 			uint32 MaxIndex = TNumericLimits< uint32 >::Min();
-			for (int32 TriangleIndex = 0; TriangleIndex < Polygon.Triangles.Num(); ++TriangleIndex)
+			for (int32 TriangleIndex = 0; TriangleIndex < PolygonTriangles.Num(); ++TriangleIndex)
 			{
-				const FMeshTriangle& Triangle = Polygon.Triangles[TriangleIndex];
+				const FMeshTriangle& Triangle = PolygonTriangles[TriangleIndex];
 				for (int32 TriVert = 0; TriVert < 3; ++TriVert)
 				{
-					const FVertexInstanceID& VertexInstanceID = Triangle.GetVertexInstanceID(TriVert);
-					int32 VertexInstanceValue = VertexInstanceID.GetValue();
-					const FMeshVertexInstance& VertexInstance = VertexInstances[VertexInstanceID];
-					if (VertexInstance.Color != FLinearColor::White)
+					const FVertexInstanceID VertexInstanceID = Triangle.GetVertexInstanceID(TriVert);
+					const int32 VertexInstanceValue = VertexInstanceID.GetValue();
+
+					const FVertexID VertexID = MeshDescription->GetVertexInstanceVertex(VertexInstanceID);
+
+					const FVector& VertexPosition = VertexPositions[VertexID];
+					const FVector& VertexInstanceNormal = VertexInstanceNormals[VertexInstanceID];
+					const FVector& VertexInstanceTangent = VertexInstanceTangents[VertexInstanceID];
+					const float VertexInstanceBinormalSign = VertexInstanceBinormalSigns[VertexInstanceID];
+					const FVector4& VertexInstanceColor = VertexInstanceColors[VertexInstanceID];
+
+					const FLinearColor LinearColor(VertexInstanceColor);
+					if (LinearColor != FLinearColor::White)
 					{
 						bHasColor = true;
 					}
 
 					FStaticMeshBuildVertex StaticMeshVertex;
 
-					StaticMeshVertex.Position = Vertices[VertexInstance.VertexID].VertexPosition * LODBuildSettings.BuildScale3D;
+					StaticMeshVertex.Position = VertexPosition * LODBuildSettings.BuildScale3D;
 					const FMatrix ScaleMatrix = FScaleMatrix(LODBuildSettings.BuildScale3D).Inverse().GetTransposed();
-					StaticMeshVertex.TangentX = ScaleMatrix.TransformVector(VertexInstance.Tangent).GetSafeNormal();
-					StaticMeshVertex.TangentY = ScaleMatrix.TransformVector(FVector::CrossProduct(VertexInstance.Normal, VertexInstance.Tangent).GetSafeNormal() * VertexInstance.BinormalSign).GetSafeNormal();
-					StaticMeshVertex.TangentZ = ScaleMatrix.TransformVector(VertexInstance.Normal).GetSafeNormal();
-					StaticMeshVertex.Color = VertexInstance.Color.ToFColor(true);
-					int32 MaxNumTexCoords = FMath::Min<int32>(MAX_MESH_TEXTURE_COORDS, MAX_STATIC_TEXCOORDS);
-					for (int32 UVIndex = 0; UVIndex < MaxNumTexCoords; ++UVIndex)
+					StaticMeshVertex.TangentX = ScaleMatrix.TransformVector(VertexInstanceTangent).GetSafeNormal();
+					StaticMeshVertex.TangentY = ScaleMatrix.TransformVector(FVector::CrossProduct(VertexInstanceNormal, VertexInstanceTangent).GetSafeNormal() * VertexInstanceBinormalSign).GetSafeNormal();
+					StaticMeshVertex.TangentZ = ScaleMatrix.TransformVector(VertexInstanceNormal).GetSafeNormal();
+					StaticMeshVertex.Color = LinearColor.ToFColor(true);
+					const uint32 MaxNumTexCoords = FMath::Min<int32>(MAX_MESH_TEXTURE_COORDS, MAX_STATIC_TEXCOORDS);
+					for (uint32 UVIndex = 0; UVIndex < MaxNumTexCoords; ++UVIndex)
 					{
-						StaticMeshVertex.UVs[UVIndex] = VertexInstance.VertexUVs.IsValidIndex(UVIndex) ? VertexInstance.VertexUVs[UVIndex] : FVector2D(0.0f, 0.0f);
+						if(UVIndex < NumTextureCoord)
+						{
+							StaticMeshVertex.UVs[UVIndex] = VertexInstanceUVs.GetArrayForIndex(UVIndex)[VertexInstanceID];
+						}
+						else
+						{
+							StaticMeshVertex.UVs[UVIndex] = FVector2D(0.0f, 0.0f);
+						}
 					}
 					
 
diff --git a/Engine/Source/Runtime/MeshDescription/Private/MeshAttributes.cpp b/Engine/Source/Runtime/MeshDescription/Private/MeshAttributes.cpp
new file mode 100644
--- /dev/null
+++ b/Engine/Source/Runtime/MeshDescription/Private/MeshAttributes.cpp
@@ -0,0 +1,30 @@
+// Copyright 1998-2017 Epic Games, Inc. All Rights Reserved.
+
+#include "MeshAttributes.h"
+
+namespace MeshAttribute
+{
+	const FName Vertex::Position( "Position ");
+	const FName Vertex::CornerSharpness( "CornerSharpness" );
+
+	const FName VertexInstance::TextureCoordinate( "TextureCoordinate" );
+	const FName VertexInstance::Normal( "Normal" );
+	const FName VertexInstance::Tangent( "Tangent" );
+	const FName VertexInstance::BinormalSign( "BinormalSign" );
+	const FName VertexInstance::Color( "Color" );
+
+	const FName Edge::IsHard( "IsHard" );
+	const FName Edge::CreaseSharpness( "CreaseSharpness" );
+
+	const FName Polygon::Normal( "Normal" );
+	const FName Polygon::Tangent( "Tangent" );
+	const FName Polygon::Binormal( "Binormal" );
+	const FName Polygon::Center( "Center" );
+
+	const FName PolygonGroup::MaterialAsset( "MaterialAsset" );
+	const FName PolygonGroup::MaterialSlotName( "MaterialSlotName" );
+	const FName PolygonGroup::ImportedMaterialSlotName( "ImportedMaterialSlotName" );
+	const FName PolygonGroup::EnableCollision( "EnableCollision" );
+	const FName PolygonGroup::CastShadow( "CastShadow" );
+}
+
diff --git a/Engine/Source/Runtime/MeshDescription/Private/MeshDescription.cpp b/Engine/Source/Runtime/MeshDescription/Private/MeshDescription.cpp
--- a/Engine/Source/Runtime/MeshDescription/Private/MeshDescription.cpp
+++ b/Engine/Source/Runtime/MeshDescription/Private/MeshDescription.cpp
@@ -3,36 +3,37 @@
 #include "MeshDescription.h"
 #include "Misc/SecureHash.h"
 #include "Serialization/MemoryWriter.h"
+#include "MeshAttributes.h"
 
 UMeshDescription::UMeshDescription()
 {
 	// Add basic vertex attributes
-	VertexAttributes().RegisterAttribute<FVector>(UEditableMeshAttribute::VertexPosition());
-	VertexAttributes().RegisterAttribute<float>(UEditableMeshAttribute::VertexCornerSharpness());
+	VertexAttributes().RegisterAttribute<FVector>( MeshAttribute::Vertex::Position, 1, FVector::ZeroVector, EMeshAttributeFlags::Lerpable );
+	VertexAttributes().RegisterAttribute<float>( MeshAttribute::Vertex::CornerSharpness, 1, 0.0f, EMeshAttributeFlags::Lerpable );
 
 	// Add basic vertex instance attributes
-	VertexInstanceAttributes().RegisterAttribute<FVector2D>(UEditableMeshAttribute::VertexTextureCoordinate(), 2 );
-	VertexInstanceAttributes().RegisterAttribute<FVector>(UEditableMeshAttribute::VertexNormal());
-	VertexInstanceAttributes().RegisterAttribute<FVector>(UEditableMeshAttribute::VertexTangent());
-	VertexInstanceAttributes().RegisterAttribute<float>(UEditableMeshAttribute::VertexBinormalSign());
-	VertexInstanceAttributes().RegisterAttribute<FVector4>(UEditableMeshAttribute::VertexColor());
+	VertexInstanceAttributes().RegisterAttribute<FVector2D>( MeshAttribute::VertexInstance::TextureCoordinate, 2, FVector2D::ZeroVector, EMeshAttributeFlags::Lerpable );
+	VertexInstanceAttributes().RegisterAttribute<FVector>( MeshAttribute::VertexInstance::Normal, 1, FVector::ZeroVector, EMeshAttributeFlags::AutoGenerated );
+	VertexInstanceAttributes().RegisterAttribute<FVector>( MeshAttribute::VertexInstance::Tangent, 1, FVector::ZeroVector, EMeshAttributeFlags::AutoGenerated );
+	VertexInstanceAttributes().RegisterAttribute<float>( MeshAttribute::VertexInstance::BinormalSign, 1, 0.0f, EMeshAttributeFlags::AutoGenerated );
+	VertexInstanceAttributes().RegisterAttribute<FVector4>( MeshAttribute::VertexInstance::Color, 1, FVector4( 1.0f ), EMeshAttributeFlags::Lerpable );
 
 	// Add basic edge attributes
-	EdgeAttributes().RegisterAttribute<bool>(UEditableMeshAttribute::EdgeIsHard());
-	EdgeAttributes().RegisterAttribute<float>(UEditableMeshAttribute::EdgeCreaseSharpness());
+	EdgeAttributes().RegisterAttribute<bool>( MeshAttribute::Edge::IsHard, 1, false );
+	EdgeAttributes().RegisterAttribute<float>( MeshAttribute::Edge::CreaseSharpness, 1, 0.0f, EMeshAttributeFlags::Lerpable );
 
 	// Add basic polygon attributes
-	PolygonAttributes().RegisterAttribute<FVector>(UEditableMeshAttribute::PolygonNormal());
-	PolygonAttributes().RegisterAttribute<FVector>( "PolygonTangent" );
-	PolygonAttributes().RegisterAttribute<FVector>( "PolygonBinormal" );
-	PolygonAttributes().RegisterAttribute<FVector>(UEditableMeshAttribute::PolygonCenter());
+	PolygonAttributes().RegisterAttribute<FVector>( MeshAttribute::Polygon::Normal, 1, FVector::ZeroVector, EMeshAttributeFlags::AutoGenerated );
+	PolygonAttributes().RegisterAttribute<FVector>( MeshAttribute::Polygon::Tangent, 1, FVector::ZeroVector, EMeshAttributeFlags::AutoGenerated );
+	PolygonAttributes().RegisterAttribute<FVector>( MeshAttribute::Polygon::Binormal, 1, FVector::ZeroVector, EMeshAttributeFlags::AutoGenerated );
+	PolygonAttributes().RegisterAttribute<FVector>( MeshAttribute::Polygon::Center, 1, FVector::ZeroVector, EMeshAttributeFlags::AutoGenerated );
 
 	// Add basic polygon group attributes
-	PolygonGroupAttributes().RegisterAttribute<FStringAssetReference>( "MaterialAsset" );
-	PolygonGroupAttributes().RegisterAttribute<FName>( "MaterialSlotName" );
-	PolygonGroupAttributes().RegisterAttribute<FName>( "ImportedMaterialSlotName" );
-	PolygonGroupAttributes().RegisterAttribute<bool>( "EnableCollision" );
-	PolygonGroupAttributes().RegisterAttribute<bool>( "CastShadow" );
+	PolygonGroupAttributes().RegisterAttribute<UObject*>( MeshAttribute::PolygonGroup::MaterialAsset );
+	PolygonGroupAttributes().RegisterAttribute<FName>( MeshAttribute::PolygonGroup::MaterialSlotName );
+	PolygonGroupAttributes().RegisterAttribute<FName>( MeshAttribute::PolygonGroup::ImportedMaterialSlotName );
+	PolygonGroupAttributes().RegisterAttribute<bool>( MeshAttribute::PolygonGroup::EnableCollision );
+	PolygonGroupAttributes().RegisterAttribute<bool>( MeshAttribute::PolygonGroup::CastShadow );
 }
 
 
@@ -53,6 +54,7 @@ void UMeshDescription::Serialize( FArchive& Ar )
 	Ar << PolygonGroupAttributesSet;
 }
 
+
 #if WITH_EDITORONLY_DATA
 FString UMeshDescription::GetIdString()
 {
@@ -77,6 +79,7 @@ FString UMeshDescription::GetIdString()
 }
 #endif
 
+
 void UMeshDescription::Compact( FElementIDRemappings& OutRemappings )
 {
 	VertexArray.Compact( OutRemappings.NewVertexIndexLookup );
@@ -85,9 +88,7 @@ void UMeshDescription::Compact( FElementIDRemappings& OutRemappings )
 	PolygonArray.Compact( OutRemappings.NewPolygonIndexLookup );
 	PolygonGroupArray.Compact( OutRemappings.NewPolygonGroupIndexLookup );
 
-	// @todo mesh description: compact attribute arrays according to these remappings
-	//
-
+	RemapAttributes( OutRemappings );
 	FixUpElementIDs( OutRemappings );
 }
 
@@ -100,9 +101,7 @@ void UMeshDescription::Remap( const FElementIDRemappings& Remappings )
 	PolygonArray.Remap( Remappings.NewPolygonIndexLookup );
 	PolygonGroupArray.Remap( Remappings.NewPolygonGroupIndexLookup );
 
-	// @todo mesh description: remap attribute arrays according to these remappings
-	//
-
+	RemapAttributes( Remappings );
 	FixUpElementIDs( Remappings );
 }
 
@@ -197,3 +196,13 @@ void UMeshDescription::FixUpElementIDs( const FElementIDRemappings& Remappings )
 		}
 	}
 }
+
+
+void UMeshDescription::RemapAttributes( const FElementIDRemappings& Remappings )
+{
+	VertexAttributesSet.Remap( Remappings.NewVertexIndexLookup );
+	VertexInstanceAttributesSet.Remap( Remappings.NewVertexInstanceIndexLookup );
+	EdgeAttributesSet.Remap( Remappings.NewEdgeIndexLookup );
+	PolygonAttributesSet.Remap( Remappings.NewPolygonIndexLookup );
+	PolygonGroupAttributesSet.Remap( Remappings.NewPolygonGroupIndexLookup );
+}
diff --git a/Engine/Source/Runtime/MeshDescription/Private/MeshTypes.cpp b/Engine/Source/Runtime/MeshDescription/Private/MeshTypes.cpp
--- a/Engine/Source/Runtime/MeshDescription/Private/MeshTypes.cpp
+++ b/Engine/Source/Runtime/MeshDescription/Private/MeshTypes.cpp
@@ -9,6 +9,7 @@ const FEdgeID FEdgeID::Invalid( TNumericLimits<uint32>::Max() );
 const FPolygonGroupID FPolygonGroupID::Invalid( TNumericLimits<uint32>::Max() );
 const FPolygonID FPolygonID::Invalid( TNumericLimits<uint32>::Max() );
 
+#if 0
 const FName UEditableMeshAttribute::VertexPositionName( "VertexPosition" );
 const FName UEditableMeshAttribute::VertexCornerSharpnessName( "VertexCornerSharpness" );
 const FName UEditableMeshAttribute::VertexNormalName( "VertexNormal" );
@@ -20,4 +21,4 @@ const FName UEditableMeshAttribute::EdgeIsHardName( "EdgeIsHard" );
 const FName UEditableMeshAttribute::EdgeCreaseSharpnessName( "EdgeCreaseSharpness" );
 const FName UEditableMeshAttribute::PolygonNormalName( "PolygonNormal" );
 const FName UEditableMeshAttribute::PolygonCenterName( "PolygonCenter" );
-
+#endif
