Make the mesh description conversion from/to rawmesh
Move triangulation from editable mesh to MeshDescription
Fix the edge creation when importing a fbx to mesh description
Add automation tests for mesh description, add some conversion tests
Load orion monolith02 level succesfully using meshdescription instead of rawmesh

diff --git a/Engine/Plugins/Editor/SpeedTreeImporter/Source/SpeedTreeImporter/Private/SpeedTreeImportFactory.cpp b/Engine/Plugins/Editor/SpeedTreeImporter/Source/SpeedTreeImporter/Private/SpeedTreeImportFactory.cpp
--- a/Engine/Plugins/Editor/SpeedTreeImporter/Source/SpeedTreeImporter/Private/SpeedTreeImportFactory.cpp
+++ b/Engine/Plugins/Editor/SpeedTreeImporter/Source/SpeedTreeImporter/Private/SpeedTreeImportFactory.cpp
@@ -1355,7 +1355,7 @@ UObject* USpeedTreeImportFactory::FactoryCreateBinary(UClass* InClass, UObject*
 						LODModel->BuildSettings.bUseFullPrecisionUVs = false;
 						LODModel->BuildSettings.bGenerateLightmapUVs = false;
 						LODModel->ScreenSize = 0.1f / FMath::Pow(2.0f, StaticMesh->SourceModels.Num() - 1);
-						LODModel->RawMeshBulkData->SaveRawMesh(RawMesh);
+						LODModel->SaveRawMesh(RawMesh);
 
 						for (int32 MaterialIndex = 0; MaterialIndex < StaticMesh->StaticMaterials.Num(); ++MaterialIndex)
 						{
@@ -1468,7 +1468,7 @@ UObject* USpeedTreeImportFactory::FactoryCreateBinary(UClass* InClass, UObject*
 					LODModel->BuildSettings.bUseFullPrecisionUVs = false;
 					LODModel->BuildSettings.bGenerateLightmapUVs = false;
 					LODModel->ScreenSize = 0.1f / FMath::Pow(2.0f, StaticMesh->SourceModels.Num() - 1);
-					LODModel->RawMeshBulkData->SaveRawMesh(RawMesh);
+					LODModel->SaveRawMesh(RawMesh);
 					// Add mesh section info entry for billboard LOD (only one section/material index)
 					const int32 LODIndex = StaticMesh->SourceModels.Num() - 1;
 					FMeshSectionInfo Info = StaticMesh->SectionInfoMap.Get(LODIndex, 0);
diff --git a/Engine/Plugins/Editor/USDImporter/Source/USDImporter/Private/StaticMeshImporter.cpp b/Engine/Plugins/Editor/USDImporter/Source/USDImporter/Private/StaticMeshImporter.cpp
--- a/Engine/Plugins/Editor/USDImporter/Source/USDImporter/Private/StaticMeshImporter.cpp
+++ b/Engine/Plugins/Editor/USDImporter/Source/USDImporter/Private/StaticMeshImporter.cpp
@@ -185,7 +185,7 @@ UStaticMesh* FUSDStaticMeshImporter::ImportStaticMesh(FUsdImportContext& ImportC
 
 			RawTriangles.CompactMaterialIndices();
 
-			SrcModel.RawMeshBulkData->SaveRawMesh(RawTriangles);
+			SrcModel.SaveRawMesh(RawTriangles);
 
 			// Recompute normals if we didnt import any
 			SrcModel.BuildSettings.bRecomputeNormals = RawTriangles.WedgeTangentZ.Num() == 0;
diff --git a/Engine/Plugins/Experimental/AlembicImporter/Source/AlembicLibrary/Private/AbcImporter.cpp b/Engine/Plugins/Experimental/AlembicImporter/Source/AlembicLibrary/Private/AbcImporter.cpp
--- a/Engine/Plugins/Experimental/AlembicImporter/Source/AlembicLibrary/Private/AbcImporter.cpp
+++ b/Engine/Plugins/Experimental/AlembicImporter/Source/AlembicLibrary/Private/AbcImporter.cpp
@@ -841,7 +841,7 @@ UStaticMesh* FAbcImporter::CreateStaticMeshFromRawMesh(UObject* InParent, const
 		SrcModel.BuildSettings.DstLightmapIndex = 1;
 
 		// Store the raw mesh within the RawMeshBulkData
-		SrcModel.RawMeshBulkData->SaveRawMesh(RawMesh);
+		SrcModel.SaveRawMesh(RawMesh);
 		
 		//Set the Imported version before calling the build
 		StaticMesh->ImportVersion = EImportStaticMeshVersion::LastVersion;
diff --git a/Engine/Plugins/Runtime/EditableMesh/Source/EditableMesh/EditableStaticMeshAdapter.cpp b/Engine/Plugins/Runtime/EditableMesh/Source/EditableMesh/EditableStaticMeshAdapter.cpp
--- a/Engine/Plugins/Runtime/EditableMesh/Source/EditableMesh/EditableStaticMeshAdapter.cpp
+++ b/Engine/Plugins/Runtime/EditableMesh/Source/EditableMesh/EditableStaticMeshAdapter.cpp
@@ -1700,7 +1700,12 @@ void UEditableStaticMeshAdapter::OnCreatePolygonGroups( const UEditableMesh* Edi
 			StaticMeshSection.bCastShadow = PolygonGroupCastShadow[ PolygonGroupID ];
 
 			UMaterialInterface* Material = Cast<UMaterialInterface>( PolygonGroupMaterialAssets[ PolygonGroupID ].TryLoad() );
-			StaticMeshSection.MaterialIndex = StaticMesh->StaticMaterials.Emplace( Material, PolygonGroupMaterialSlotNames[ PolygonGroupID ], PolygonGroupImportedMaterialSlotNames[ PolygonGroupID ] );
+#if WITH_EDITORONLY_DATA
+			StaticMeshSection.MaterialIndex = StaticMesh->StaticMaterials.Emplace(Material, PolygonGroupMaterialSlotNames[PolygonGroupID], PolygonGroupImportedMaterialSlotNames[PolygonGroupID]);
+#else
+			StaticMeshSection.MaterialIndex = StaticMesh->StaticMaterials.Emplace(Material, PolygonGroupMaterialSlotNames[PolygonGroupID]);
+#endif
+			
 		}
 
 		// Insert the rendering polygon group for keeping track of these index buffer properties
@@ -1730,7 +1735,11 @@ void UEditableStaticMeshAdapter::OnSetPolygonGroupAttribute( const UEditableMesh
 		Attribute.AttributeName == MeshAttribute::PolygonGroup::ImportedMaterialSlotName )
 	{
 		UMaterialInterface* Material = Cast<UMaterialInterface>( PolygonGroupMaterialAssets[ PolygonGroupID ].TryLoad() );
+#if WITH_EDITORONLY_DATA
 		StaticMesh->StaticMaterials[ StaticMeshSection.MaterialIndex ] = FStaticMaterial( Material, PolygonGroupMaterialSlotNames[ PolygonGroupID ], PolygonGroupImportedMaterialSlotNames[ PolygonGroupID ] );
+#else
+		StaticMesh->StaticMaterials[StaticMeshSection.MaterialIndex] = FStaticMaterial(Material, PolygonGroupMaterialSlotNames[PolygonGroupID]);
+#endif
 	}
 	else if( Attribute.AttributeName == MeshAttribute::PolygonGroup::CastShadow )
 	{
diff --git a/Engine/Plugins/Runtime/ProceduralMeshComponent/Source/ProceduralMeshComponentEditor/Private/ProceduralMeshComponentDetails.cpp b/Engine/Plugins/Runtime/ProceduralMeshComponent/Source/ProceduralMeshComponentEditor/Private/ProceduralMeshComponentDetails.cpp
--- a/Engine/Plugins/Runtime/ProceduralMeshComponent/Source/ProceduralMeshComponentEditor/Private/ProceduralMeshComponentDetails.cpp
+++ b/Engine/Plugins/Runtime/ProceduralMeshComponent/Source/ProceduralMeshComponentEditor/Private/ProceduralMeshComponentDetails.cpp
@@ -192,7 +192,7 @@ FReply FProceduralMeshComponentDetails::ClickedOnConvertToStaticMesh()
 				SrcModel->BuildSettings.bGenerateLightmapUVs = true;
 				SrcModel->BuildSettings.SrcLightmapIndex = 0;
 				SrcModel->BuildSettings.DstLightmapIndex = 1;
-				SrcModel->RawMeshBulkData->SaveRawMesh(RawMesh);
+				SrcModel->SaveRawMesh(RawMesh);
 
 				// Copy materials to new mesh
 				for (UMaterialInterface* Material : MeshMaterials)
diff --git a/Engine/Plugins/Tests/EditorTests/Source/EditorTests/Private/UnrealEd/FbxAutomationTests.cpp b/Engine/Plugins/Tests/EditorTests/Source/EditorTests/Private/UnrealEd/FbxAutomationTests.cpp
--- a/Engine/Plugins/Tests/EditorTests/Source/EditorTests/Private/UnrealEd/FbxAutomationTests.cpp
+++ b/Engine/Plugins/Tests/EditorTests/Source/EditorTests/Private/UnrealEd/FbxAutomationTests.cpp
@@ -145,6 +145,16 @@ bool FFbxImportAssetsAutomationTest::RunTest(const FString& Parameters)
 	//Read the fbx options from the .json file and fill the ImportUI
 	FbxAutomationTestsAPI::ReadFbxOptions(FileOptionAndResult, TestPlanArray);
 
+	static const auto CVarDistanceField = IConsoleManager::Get().FindTConsoleVariableDataInt(TEXT("r.GenerateMeshDistanceFields"));
+	int32 OriginalCVarDistanceFieldValue = CVarDistanceField->GetValueOnGameThread();
+	IConsoleVariable* CVarDistanceFieldInterface = IConsoleManager::Get().FindConsoleVariable(TEXT("r.GenerateMeshDistanceFields"));
+	//Avoid building the distance field when we do fbx automation test
+	if (OriginalCVarDistanceFieldValue != 0 && CVarDistanceFieldInterface)
+	{
+		//Hack we change the distance field user console variable to control the build, but we put back the value after the first build
+		CVarDistanceFieldInterface->SetWithCurrentPriority(0);
+	}
+
 	FAssetToolsModule& AssetToolsModule = FModuleManager::LoadModuleChecked<FAssetToolsModule>("AssetTools");
 	bool CurTestSuccessful = (TestPlanArray.Num() > 0);
 	TArray<UObject*> GlobalImportedObjects;
@@ -1464,6 +1474,12 @@ bool FFbxImportAssetsAutomationTest::RunTest(const FString& Parameters)
 			ObjectTools::ForceDeleteObjects(ObjectToDelete, false);
 		}
 	}
+	
+	//Put back the distance field value
+	if (OriginalCVarDistanceFieldValue != 0 && CVarDistanceFieldInterface)
+	{
+		CVarDistanceFieldInterface->SetWithCurrentPriority(OriginalCVarDistanceFieldValue);
+	}
 
 	return CurTestSuccessful;
 }
diff --git a/Engine/Plugins/Tests/EditorTests/Source/EditorTests/Private/UnrealEd/MeshDescriptionTests.cpp b/Engine/Plugins/Tests/EditorTests/Source/EditorTests/Private/UnrealEd/MeshDescriptionTests.cpp
new file mode 100644
--- /dev/null
+++ b/Engine/Plugins/Tests/EditorTests/Source/EditorTests/Private/UnrealEd/MeshDescriptionTests.cpp
@@ -0,0 +1,451 @@
+// Copyright 1998-2017 Epic Games, Inc. All Rights Reserved.
+
+#include "CoreMinimal.h"
+#include "UObject/Object.h"
+#include "UObject/UObjectGlobals.h"
+#include "UObject/Package.h"
+#include "Misc/AutomationTest.h"
+#include "Modules/ModuleManager.h"
+#include "Engine/StaticMesh.h"
+#include "Materials/Material.h"
+#include "MeshDescription.h"
+#include "MeshBuilder.h"
+#include "MeshAttributes.h"
+#include "RawMesh.h"
+
+//////////////////////////////////////////////////////////////////////////
+
+/**
+* FMeshDescriptionAutomationTest
+* Test that verify the MeshDescription functionalities. (Creation, modification, conversion to/from FRawMesh, render build)
+* The tests will create some transient geometry using the mesh description API
+* Cannot be run in a commandlet as it executes code that routes through Slate UI.
+*/
+IMPLEMENT_COMPLEX_AUTOMATION_TEST(FMeshDescriptionAutomationTest, "Editor.Meshes.MeshDescription", (EAutomationTestFlags::EditorContext | EAutomationTestFlags::NonNullRHI | EAutomationTestFlags::EngineFilter))
+
+#define ConversionTestData 1
+
+class FMeshDescriptionTest
+{
+public:
+	FMeshDescriptionTest(const FString& InBeautifiedNames, const FString& InTestData)
+		: BeautifiedNames(InBeautifiedNames)
+		, TestData(InTestData)
+	{}
+
+	FString BeautifiedNames;
+	FString TestData;
+	bool Execute(FAutomationTestExecutionInfo& ExecutionInfo);
+	bool ConversionTest(FAutomationTestExecutionInfo& ExecutionInfo);
+private:
+	bool CompareMeshDescription(const FString& AssetName, FAutomationTestExecutionInfo& ExecutionInfo, const UMeshDescription* ReferenceMeshDescription, const UMeshDescription* MeshDescription);
+};
+
+class FMeshDescriptionTests
+{
+public:
+	static FMeshDescriptionTests& GetInstance()
+	{
+		static FMeshDescriptionTests Instance;
+		return Instance;
+	}
+
+	void ClearTests()
+	{
+		AllTests.Empty();
+	}
+
+	bool ExecTest(int32 TestKey, FAutomationTestExecutionInfo& ExecutionInfo)
+	{
+		check(AllTests.Contains(TestKey));
+		FMeshDescriptionTest& Test = AllTests[TestKey];
+		return Test.Execute(ExecutionInfo);
+	}
+
+	bool AddTest(const FMeshDescriptionTest& MeshDescriptionTest)
+	{
+		if (!MeshDescriptionTest.TestData.IsNumeric())
+		{
+			return false;
+		}
+		int32 TestID = FPlatformString::Atoi(*(MeshDescriptionTest.TestData));
+		check(!AllTests.Contains(TestID));
+		AllTests.Add(TestID, MeshDescriptionTest);
+		return true;
+	}
+
+private:
+	TMap<int32, FMeshDescriptionTest> AllTests;
+	FMeshDescriptionTests() {}
+};
+
+/**
+* Requests a enumeration of all sample assets to import
+*/
+void FMeshDescriptionAutomationTest::GetTests(TArray<FString>& OutBeautifiedNames, TArray <FString>& OutTestCommands) const
+{
+	FMeshDescriptionTests::GetInstance().ClearTests();
+	//Create cube test
+	FMeshDescriptionTest ConversionTest(TEXT("Conversion test data"), FString::FromInt(ConversionTestData));
+	FMeshDescriptionTests::GetInstance().AddTest(ConversionTest);
+	OutBeautifiedNames.Add(ConversionTest.BeautifiedNames);
+	OutTestCommands.Add(ConversionTest.TestData);
+}
+
+/**
+* Execute the generic import test
+*
+* @param Parameters - Should specify the asset to import
+* @return	TRUE if the test was successful, FALSE otherwise
+*/
+bool FMeshDescriptionAutomationTest::RunTest(const FString& Parameters)
+{
+	if (!Parameters.IsNumeric())
+	{
+		ExecutionInfo.AddEvent(FAutomationEvent(EAutomationEventType::Error, FString::Printf(TEXT("Wrong parameter for mesh description test parameter should be a number: [%s]"), *Parameters)));
+		return false;
+	}
+	int32 TestID = FPlatformString::Atoi(*Parameters);
+	return FMeshDescriptionTests::GetInstance().ExecTest(TestID, ExecutionInfo);
+}
+
+bool FMeshDescriptionTest::Execute(FAutomationTestExecutionInfo& ExecutionInfo)
+{
+	if (!TestData.IsNumeric())
+	{
+		ExecutionInfo.AddEvent(FAutomationEvent(EAutomationEventType::Error, FString::Printf(TEXT("Wrong parameter for mesh description test parameter should be a number: [%s]"), *TestData)));
+		return false;
+	}
+	int32 TestID = FPlatformString::Atoi(*TestData);
+	bool bSuccess = false;
+	switch (TestID)
+	{
+	case ConversionTestData:
+		bSuccess = ConversionTest(ExecutionInfo);
+		break;
+	}
+	return bSuccess;
+}
+
+bool FMeshDescriptionTest::CompareMeshDescription(const FString& AssetName, FAutomationTestExecutionInfo& ExecutionInfo, const UMeshDescription* ReferenceMeshDescription, const UMeshDescription* MeshDescription)
+{
+	//////////////////////////////////////////////////////////////////////////
+	//Gather the reference data
+	const TVertexAttributeArray<FVector>& ReferenceVertexPositions = ReferenceMeshDescription->VertexAttributes().GetAttributes<FVector>(MeshAttribute::Vertex::Position);
+	//const TVertexAttributeArray<float>& ReferenceVertexCornerSharpness = ReferenceMeshDescription->VertexAttributes().GetAttributes<float>(MeshAttribute::Vertex::CornerSharpness);
+
+	const TVertexInstanceAttributeArray<FVector>& ReferenceVertexInstanceNormals = ReferenceMeshDescription->VertexInstanceAttributes().GetAttributes<FVector>(MeshAttribute::VertexInstance::Normal);
+	const TVertexInstanceAttributeArray<FVector>& ReferenceVertexInstanceTangents = ReferenceMeshDescription->VertexInstanceAttributes().GetAttributes<FVector>(MeshAttribute::VertexInstance::Tangent);
+	const TVertexInstanceAttributeArray<float>& ReferenceVertexInstanceBinormalSigns = ReferenceMeshDescription->VertexInstanceAttributes().GetAttributes<float>(MeshAttribute::VertexInstance::BinormalSign);
+	const TVertexInstanceAttributeArray<FVector4>& ReferenceVertexInstanceColors = ReferenceMeshDescription->VertexInstanceAttributes().GetAttributes<FVector4>(MeshAttribute::VertexInstance::Color);
+	const TVertexInstanceAttributeIndicesArray<FVector2D>& ReferenceVertexInstanceUVs = ReferenceMeshDescription->VertexInstanceAttributes().GetAttributesSet<FVector2D>(MeshAttribute::VertexInstance::TextureCoordinate);
+
+	const TEdgeAttributeArray<bool>& ReferenceEdgeHardnesses = ReferenceMeshDescription->EdgeAttributes().GetAttributes<bool>(MeshAttribute::Edge::IsHard);
+	//const TEdgeAttributeArray<float>& ReferenceEdgeCreaseSharpnesses = ReferenceMeshDescription->EdgeAttributes().GetAttributes<float>(MeshAttribute::Edge::CreaseSharpness);
+
+	//const TPolygonAttributeArray<FVector>& ReferencePolygonNormals = ReferenceMeshDescription->PolygonAttributes().GetAttributes<FVector>(MeshAttribute::Polygon::Normal);
+	//const TPolygonAttributeArray<FVector>& ReferencePolygonTangents = ReferenceMeshDescription->PolygonAttributes().GetAttributes<FVector>(MeshAttribute::Polygon::Tangent);
+	//const TPolygonAttributeArray<FVector>& ReferencePolygonBiNormals = ReferenceMeshDescription->PolygonAttributes().GetAttributes<FVector>(MeshAttribute::Polygon::Binormal);
+	//const TPolygonAttributeArray<FVector>& ReferencePolygonCenters = ReferenceMeshDescription->PolygonAttributes().GetAttributes<FVector>(MeshAttribute::Polygon::Center);
+
+	const TPolygonGroupAttributeArray<FSoftObjectPath>& ReferencePolygonGroupMaterialAssets = ReferenceMeshDescription->PolygonGroupAttributes().GetAttributes<FSoftObjectPath>(MeshAttribute::PolygonGroup::MaterialAsset);
+	const TPolygonGroupAttributeArray<FName>& ReferencePolygonGroupMaterialSlotNames = ReferenceMeshDescription->PolygonGroupAttributes().GetAttributes<FName>(MeshAttribute::PolygonGroup::MaterialSlotName);
+	const TPolygonGroupAttributeArray<FName>& ReferencePolygonGroupImportedMaterialSlotNames = ReferenceMeshDescription->PolygonGroupAttributes().GetAttributes<FName>(MeshAttribute::PolygonGroup::ImportedMaterialSlotName);
+	const TPolygonGroupAttributeArray<bool>& ReferencePolygonGroupCollision = ReferenceMeshDescription->PolygonGroupAttributes().GetAttributes<bool>(MeshAttribute::PolygonGroup::EnableCollision);
+	const TPolygonGroupAttributeArray<bool>& ReferencePolygonGroupCastShadow = ReferenceMeshDescription->PolygonGroupAttributes().GetAttributes<bool>(MeshAttribute::PolygonGroup::CastShadow);
+
+
+	//////////////////////////////////////////////////////////////////////////
+	//Gather the result data
+	const TVertexAttributeArray<FVector>& ResultVertexPositions = MeshDescription->VertexAttributes().GetAttributes<FVector>(MeshAttribute::Vertex::Position);
+	//const TVertexAttributeArray<float>& ResultVertexCornerSharpness = MeshDescription->VertexAttributes().GetAttributes<float>(MeshAttribute::Vertex::CornerSharpness);
+
+	const TVertexInstanceAttributeArray<FVector>& ResultVertexInstanceNormals = MeshDescription->VertexInstanceAttributes().GetAttributes<FVector>(MeshAttribute::VertexInstance::Normal);
+	const TVertexInstanceAttributeArray<FVector>& ResultVertexInstanceTangents = MeshDescription->VertexInstanceAttributes().GetAttributes<FVector>(MeshAttribute::VertexInstance::Tangent);
+	const TVertexInstanceAttributeArray<float>& ResultVertexInstanceBinormalSigns = MeshDescription->VertexInstanceAttributes().GetAttributes<float>(MeshAttribute::VertexInstance::BinormalSign);
+	const TVertexInstanceAttributeArray<FVector4>& ResultVertexInstanceColors = MeshDescription->VertexInstanceAttributes().GetAttributes<FVector4>(MeshAttribute::VertexInstance::Color);
+	const TVertexInstanceAttributeIndicesArray<FVector2D>& ResultVertexInstanceUVs = MeshDescription->VertexInstanceAttributes().GetAttributesSet<FVector2D>(MeshAttribute::VertexInstance::TextureCoordinate);
+
+	const TEdgeAttributeArray<bool>& ResultEdgeHardnesses = MeshDescription->EdgeAttributes().GetAttributes<bool>(MeshAttribute::Edge::IsHard);
+	//const TEdgeAttributeArray<float>& ResultEdgeCreaseSharpnesses = MeshDescription->EdgeAttributes().GetAttributes<float>(MeshAttribute::Edge::CreaseSharpness);
+
+	//const TPolygonAttributeArray<FVector>& ResultPolygonNormals = MeshDescription->PolygonAttributes().GetAttributes<FVector>(MeshAttribute::Polygon::Normal);
+	//const TPolygonAttributeArray<FVector>& ResultPolygonTangents = MeshDescription->PolygonAttributes().GetAttributes<FVector>(MeshAttribute::Polygon::Tangent);
+	//const TPolygonAttributeArray<FVector>& ResultPolygonBiNormals = MeshDescription->PolygonAttributes().GetAttributes<FVector>(MeshAttribute::Polygon::Binormal);
+	//const TPolygonAttributeArray<FVector>& ResultPolygonCenters = MeshDescription->PolygonAttributes().GetAttributes<FVector>(MeshAttribute::Polygon::Center);
+
+	const TPolygonGroupAttributeArray<FSoftObjectPath>& ResultPolygonGroupMaterialAssets = MeshDescription->PolygonGroupAttributes().GetAttributes<FSoftObjectPath>(MeshAttribute::PolygonGroup::MaterialAsset);
+	const TPolygonGroupAttributeArray<FName>& ResultPolygonGroupMaterialSlotNames = MeshDescription->PolygonGroupAttributes().GetAttributes<FName>(MeshAttribute::PolygonGroup::MaterialSlotName);
+	const TPolygonGroupAttributeArray<FName>& ResultPolygonGroupImportedMaterialSlotNames = MeshDescription->PolygonGroupAttributes().GetAttributes<FName>(MeshAttribute::PolygonGroup::ImportedMaterialSlotName);
+	const TPolygonGroupAttributeArray<bool>& ResultPolygonGroupCollision = MeshDescription->PolygonGroupAttributes().GetAttributes<bool>(MeshAttribute::PolygonGroup::EnableCollision);
+	const TPolygonGroupAttributeArray<bool>& ResultPolygonGroupCastShadow = MeshDescription->PolygonGroupAttributes().GetAttributes<bool>(MeshAttribute::PolygonGroup::CastShadow);
+
+	
+	//////////////////////////////////////////////////////////////////////////
+	// Do the comparison
+	bool bAllSame = true;
+	//Positions
+	if (ReferenceVertexPositions.Num() != ResultVertexPositions.Num())
+	{
+		ExecutionInfo.AddEvent(FAutomationEvent(EAutomationEventType::Error, FString::Printf(TEXT("The %s conversion to RawMesh is not lossless, vertex positions count is different. Vertex count expected [%d] result [%d]"),
+			*AssetName,
+			ReferenceVertexPositions.Num(),
+			ResultVertexPositions.Num())));
+		bAllSame = false;
+	}
+	else
+	{
+		for (FVertexID VertexID : ReferenceMeshDescription->Vertices().GetElementIDs())
+		{
+			if (ReferenceVertexPositions[VertexID] != ResultVertexPositions[VertexID])
+			{
+				ExecutionInfo.AddEvent(FAutomationEvent(EAutomationEventType::Error, FString::Printf(TEXT("The %s conversion to RawMesh is not lossless, vertex position array is different. VertexID [%d] expected position [%s] result [%s]"),
+					*AssetName,
+					VertexID.GetValue(),
+					*ReferenceVertexPositions[VertexID].ToString(),
+					*ResultVertexPositions[VertexID].ToString())));
+				bAllSame = false;
+				break;
+			}
+		}
+	}
+	//Normals
+	if (ReferenceVertexInstanceNormals.Num() != ResultVertexInstanceNormals.Num())
+	{
+		ExecutionInfo.AddEvent(FAutomationEvent(EAutomationEventType::Error, FString::Printf(TEXT("The %s conversion to RawMesh is not lossless, vertex instance normals count is different. Normals count expected [%d] result [%d]"),
+			*AssetName,
+			ReferenceVertexInstanceNormals.Num(),
+			ResultVertexInstanceNormals.Num())));
+		bAllSame = false;
+	}
+	else
+	{
+		for (FVertexInstanceID VertexInstanceID : ReferenceMeshDescription->VertexInstances().GetElementIDs())
+		{
+			if (ReferenceVertexInstanceNormals[VertexInstanceID] != ResultVertexInstanceNormals[VertexInstanceID])
+			{
+				ExecutionInfo.AddEvent(FAutomationEvent(EAutomationEventType::Error, FString::Printf(TEXT("The %s conversion to RawMesh is not lossless, vertex instance normals array is different. VertexInstanceID [%d] expected normal [%s] result [%s]"),
+					*AssetName,
+					VertexInstanceID.GetValue(),
+					*ReferenceVertexInstanceNormals[VertexInstanceID].ToString(),
+					*ResultVertexInstanceNormals[VertexInstanceID].ToString())));
+				bAllSame = false;
+				break;
+			}
+		}
+	}
+	//Tangents
+	if (ReferenceVertexInstanceTangents.Num() != ResultVertexInstanceTangents.Num())
+	{
+		ExecutionInfo.AddEvent(FAutomationEvent(EAutomationEventType::Error, FString::Printf(TEXT("The %s conversion to RawMesh is not lossless, vertex instance Tangents count is different. Tangents count expected [%d] result [%d]"),
+			*AssetName,
+			ReferenceVertexInstanceTangents.Num(),
+			ResultVertexInstanceTangents.Num())));
+		bAllSame = false;
+	}
+	else
+	{
+		for (FVertexInstanceID VertexInstanceID : ReferenceMeshDescription->VertexInstances().GetElementIDs())
+		{
+			if (ReferenceVertexInstanceTangents[VertexInstanceID] != ResultVertexInstanceTangents[VertexInstanceID])
+			{
+				ExecutionInfo.AddEvent(FAutomationEvent(EAutomationEventType::Error, FString::Printf(TEXT("The %s conversion to RawMesh is not lossless, vertex instance Tangents array is different. VertexInstanceID [%d] expected Tangent [%s] result [%s]"),
+					*AssetName,
+					VertexInstanceID.GetValue(),
+					*ReferenceVertexInstanceTangents[VertexInstanceID].ToString(),
+					*ResultVertexInstanceTangents[VertexInstanceID].ToString())));
+				bAllSame = false;
+				break;
+			}
+		}
+	}
+	//BiNormal signs
+	if (ReferenceVertexInstanceBinormalSigns.Num() != ResultVertexInstanceBinormalSigns.Num())
+	{
+		ExecutionInfo.AddEvent(FAutomationEvent(EAutomationEventType::Error, FString::Printf(TEXT("The %s conversion to RawMesh is not lossless, vertex instance BinormalSigns count is different. BinormalSigns count expected [%d] result [%d]"),
+			*AssetName,
+			ReferenceVertexInstanceBinormalSigns.Num(),
+			ResultVertexInstanceBinormalSigns.Num())));
+		bAllSame = false;
+	}
+	else
+	{
+		for (FVertexInstanceID VertexInstanceID : ReferenceMeshDescription->VertexInstances().GetElementIDs())
+		{
+			if (ReferenceVertexInstanceBinormalSigns[VertexInstanceID] != ResultVertexInstanceBinormalSigns[VertexInstanceID])
+			{
+				ExecutionInfo.AddEvent(FAutomationEvent(EAutomationEventType::Error, FString::Printf(TEXT("The %s conversion to RawMesh is not lossless, vertex instance BinormalSigns array is different. VertexInstanceID [%d] expected binormal sign [%f] result [%f]"),
+					*AssetName,
+					VertexInstanceID.GetValue(),
+					ReferenceVertexInstanceBinormalSigns[VertexInstanceID],
+					ResultVertexInstanceBinormalSigns[VertexInstanceID])));
+				bAllSame = false;
+				break;
+			}
+		}
+	}
+	//Colors
+	if (ReferenceVertexInstanceColors.Num() != ResultVertexInstanceColors.Num())
+	{
+		ExecutionInfo.AddEvent(FAutomationEvent(EAutomationEventType::Error, FString::Printf(TEXT("The %s conversion to RawMesh is not lossless, vertex instance Colors count is different. Colors count expected [%d] result [%d]"),
+			*AssetName,
+			ReferenceVertexInstanceColors.Num(),
+			ResultVertexInstanceColors.Num())));
+		bAllSame = false;
+	}
+	else
+	{
+		for (FVertexInstanceID VertexInstanceID : ReferenceMeshDescription->VertexInstances().GetElementIDs())
+		{
+			if (ReferenceVertexInstanceColors[VertexInstanceID] != ResultVertexInstanceColors[VertexInstanceID])
+			{
+				ExecutionInfo.AddEvent(FAutomationEvent(EAutomationEventType::Error, FString::Printf(TEXT("The %s conversion to RawMesh is not lossless, vertex instance Colors array is different. VertexInstanceID [%d] expected Color [%s] result [%s]"),
+					*AssetName, 
+					VertexInstanceID.GetValue(),
+					*ReferenceVertexInstanceColors[VertexInstanceID].ToString(),
+					*ResultVertexInstanceColors[VertexInstanceID].ToString())));
+				bAllSame = false;
+				break;
+			}
+		}
+	}
+	//Uvs
+	if (ReferenceVertexInstanceUVs.GetNumIndices() != ResultVertexInstanceUVs.GetNumIndices())
+	{
+		ExecutionInfo.AddEvent(FAutomationEvent(EAutomationEventType::Error, FString::Printf(TEXT("The %s conversion to RawMesh is not lossless, vertex instance UVs channel count is different. UVs channel count expected [%d] result [%d]"),
+			*AssetName,
+			ReferenceVertexInstanceUVs.GetNumIndices(),
+			ResultVertexInstanceUVs.GetNumIndices())));
+		bAllSame = false;
+	}
+	else
+	{
+		for (int32 UVChannelIndex = 0; UVChannelIndex < ReferenceVertexInstanceUVs.GetNumIndices(); ++UVChannelIndex)
+		{
+			const TMeshAttributeArray<FVector2D, FVertexInstanceID>& ReferenceUVs = ReferenceVertexInstanceUVs.GetArrayForIndex(UVChannelIndex);
+			const TMeshAttributeArray<FVector2D, FVertexInstanceID>& ResultUVs = ResultVertexInstanceUVs.GetArrayForIndex(UVChannelIndex);
+			if (ReferenceUVs.Num() != ResultUVs.Num())
+			{
+				ExecutionInfo.AddEvent(FAutomationEvent(EAutomationEventType::Error, FString::Printf(TEXT("The %s conversion to RawMesh is not lossless, vertex instance UV channel[%d] UVs count is different. Uvs count expected [%d] result [%d]"),
+					*AssetName,
+					UVChannelIndex,
+					ReferenceUVs.Num(),
+					ResultUVs.Num())));
+				bAllSame = false;
+			}
+			else
+			{
+				for (FVertexInstanceID VertexInstanceID : ReferenceMeshDescription->VertexInstances().GetElementIDs())
+				{
+					if (ReferenceUVs[VertexInstanceID] != ResultUVs[VertexInstanceID])
+					{
+						ExecutionInfo.AddEvent(FAutomationEvent(EAutomationEventType::Error, FString::Printf(TEXT("The %s conversion to RawMesh is not lossless, vertex instance UV Channel[%d] UVs array is different. VertexInstanceID [%d] expected UV [%s] result [%s]"),
+							*AssetName,
+							UVChannelIndex,
+							VertexInstanceID.GetValue(),
+							*ReferenceUVs[VertexInstanceID].ToString(),
+							*ResultUVs[VertexInstanceID].ToString())));
+						bAllSame = false;
+						break;
+					}
+				}
+			}
+		}
+	}
+	//Polygon group ID
+	//We currently rely only on the PolygonGroupID. The duplicate material slot information (the info is store in the staticmesh material array) is not necessary and cannot be put in FRawMesh structure.
+	if (ReferenceMeshDescription->PolygonGroups().Num() != MeshDescription->PolygonGroups().Num())
+	{
+		ExecutionInfo.AddEvent(FAutomationEvent(EAutomationEventType::Error, FString::Printf(TEXT("The %s conversion to RawMesh is not lossless, Polygon group count is different. Polygon group count expected [%d] result [%d]"),
+			*AssetName,
+			ReferenceMeshDescription->PolygonGroups().Num(),
+			MeshDescription->PolygonGroups().Num())));
+		bAllSame = false;
+	}
+	//Edges
+	//We check if hard edges are kept correctly
+	if (ReferenceEdgeHardnesses.Num() != ResultEdgeHardnesses.Num())
+	{
+		ExecutionInfo.AddEvent(FAutomationEvent(EAutomationEventType::Error, FString::Printf(TEXT("The %s conversion to RawMesh is not lossless, Edge count is different. Edges count expected [%d] result [%d]"),
+			*AssetName,
+			ReferenceEdgeHardnesses.Num(),
+			ResultEdgeHardnesses.Num())));
+		bAllSame = false;
+	}
+	else
+	{
+		for (const FEdgeID& EdgeID : MeshDescription->Edges().GetElementIDs())
+		{
+			if (ReferenceEdgeHardnesses[EdgeID] != ResultEdgeHardnesses[EdgeID])
+			{
+				//Make sure it is not an external edge (only one polygon connected) since it is impossible to retain this information in a smoothing group.
+				//External edge hardnesses has no impact on the normal calculation. It is useful only when editing meshes.
+				const TArray<FPolygonID>& EdgeConnectedPolygons = MeshDescription->GetEdgeConnectedPolygons(EdgeID);
+				if (EdgeConnectedPolygons.Num() > 1)
+				{
+					ExecutionInfo.AddEvent(FAutomationEvent(EAutomationEventType::Error, FString::Printf(TEXT("The %s conversion to RawMesh is not lossless, Edge hardnesses array is different. EdgeID [%d] expected hardnesse [%s] result [%s]"),
+						*AssetName,
+						EdgeID.GetValue(),
+						(ReferenceEdgeHardnesses[EdgeID] ? TEXT("true") : TEXT("false")),
+						(ResultEdgeHardnesses[EdgeID] ? TEXT("true") : TEXT("false")))));
+					bAllSame = false;
+				}
+				//break;
+			}
+		}
+	}
+
+	return bAllSame;
+}
+
+bool FMeshDescriptionTest::ConversionTest(FAutomationTestExecutionInfo& ExecutionInfo)
+{
+	TArray<FString> AssetNames;
+	AssetNames.Add(TEXT("Cone_1"));
+	AssetNames.Add(TEXT("Cone_2"));
+	AssetNames.Add(TEXT("Cube"));
+	AssetNames.Add(TEXT("Patch_1"));
+	AssetNames.Add(TEXT("Patch_2"));
+	AssetNames.Add(TEXT("Patch_3"));
+	AssetNames.Add(TEXT("Patch_4"));
+	AssetNames.Add(TEXT("Patch_5"));
+	AssetNames.Add(TEXT("Pentagone"));
+	AssetNames.Add(TEXT("Sphere_1"));
+	AssetNames.Add(TEXT("Sphere_2"));
+	AssetNames.Add(TEXT("Sphere_3"));
+	AssetNames.Add(TEXT("Torus_1"));
+	AssetNames.Add(TEXT("Torus_2"));
+
+	bool bAllSame = true;
+	for (const FString& AssetName : AssetNames)
+	{
+		FString FullAssetName = TEXT("/Game/Tests/MeshDescription/") + AssetName + TEXT(".") + AssetName;
+		UStaticMesh* AssetMesh = LoadObject<UStaticMesh>(nullptr, *FullAssetName, nullptr, LOAD_None, nullptr);
+
+#if UE_BUILD_DEBUG
+		AssetMesh->BuildCacheAutomationTestGuid = FGuid::NewGuid();
+#endif
+
+		if (AssetMesh != nullptr)
+		{
+			const UMeshDescription* ReferenceAssetMesh = AssetMesh->GetOriginalMeshDescription(0);
+			check(ReferenceAssetMesh != nullptr);
+			//Create a temporary Mesh Description
+			UMeshDescription* ResultAssetMesh = DuplicateObject<UMeshDescription>(ReferenceAssetMesh, GetTransientPackage(), NAME_None);
+			//Convert MeshDescription to FRawMesh
+			FRawMesh RawMesh;
+			FMeshDescriptionOperations::ConverToRawMesh(ResultAssetMesh, RawMesh);
+			//Convert back the FRawmesh
+			FMeshDescriptionOperations::ConverFromRawMesh(RawMesh, ResultAssetMesh);
+			if (!CompareMeshDescription(AssetName, ExecutionInfo, ReferenceAssetMesh, ResultAssetMesh))
+			{
+				bAllSame = false;
+			}
+			//Destroy temporary object
+			ResultAssetMesh->MarkPendingKill();
+		}
+	}
+	//Collect garbage pass to remove everything from memory
+	TryCollectGarbage(GARBAGE_COLLECTION_KEEPFLAGS);
+	return bAllSame;
+}
diff --git a/Engine/Source/Developer/MeshMergeUtilities/Private/MeshMergeHelpers.cpp b/Engine/Source/Developer/MeshMergeUtilities/Private/MeshMergeHelpers.cpp
--- a/Engine/Source/Developer/MeshMergeUtilities/Private/MeshMergeHelpers.cpp
+++ b/Engine/Source/Developer/MeshMergeUtilities/Private/MeshMergeHelpers.cpp
@@ -302,7 +302,7 @@ void FMeshMergeHelpers::RetrieveMesh(const UStaticMesh* StaticMesh, int32 LODInd
 	const FStaticMeshSourceModel& StaticMeshModel = StaticMesh->SourceModels[LODIndex];
 
 	// Imported meshes will have a filled RawMeshBulkData set
-	const bool bImportedMesh = !StaticMeshModel.RawMeshBulkData->IsEmpty();
+	const bool bImportedMesh = !StaticMeshModel.IsRawMeshEmpty();
 	// Check whether or not this mesh has been reduced in-engine
 	const bool bReducedMesh = (StaticMeshModel.ReductionSettings.PercentTriangles < 1.0f);
 	// Trying to retrieve rawmesh from SourceStaticMeshModel was giving issues, which causes a mismatch			
@@ -310,7 +310,7 @@ void FMeshMergeHelpers::RetrieveMesh(const UStaticMesh* StaticMesh, int32 LODInd
 
 	if (bImportedMesh && !bReducedMesh && !bRenderDataMismatch)
 	{
-		StaticMeshModel.RawMeshBulkData->LoadRawMesh(RawMesh);
+		StaticMeshModel.LoadRawMesh(RawMesh);
 	}
 	else
 	{
diff --git a/Engine/Source/Developer/MeshMergeUtilities/Private/MeshMergeUtilities.cpp b/Engine/Source/Developer/MeshMergeUtilities/Private/MeshMergeUtilities.cpp
--- a/Engine/Source/Developer/MeshMergeUtilities/Private/MeshMergeUtilities.cpp
+++ b/Engine/Source/Developer/MeshMergeUtilities/Private/MeshMergeUtilities.cpp
@@ -2195,7 +2195,7 @@ void FMeshMergeUtilities::MergeComponentsToStaticMesh(const TArray<UPrimitiveCom
 				SrcModel->BuildSettings.SrcLightmapIndex = 0;
 				SrcModel->BuildSettings.DstLightmapIndex = LightMapUVChannel;
 
-				SrcModel->RawMeshBulkData->SaveRawMesh(MergedMeshLOD);
+				SrcModel->SaveRawMesh(MergedMeshLOD);
 			}
 		}
 		
diff --git a/Engine/Source/Developer/MeshMergeUtilities/Private/ProxyGenerationProcessor.cpp b/Engine/Source/Developer/MeshMergeUtilities/Private/ProxyGenerationProcessor.cpp
--- a/Engine/Source/Developer/MeshMergeUtilities/Private/ProxyGenerationProcessor.cpp
+++ b/Engine/Source/Developer/MeshMergeUtilities/Private/ProxyGenerationProcessor.cpp
@@ -180,7 +180,7 @@ void FProxyGenerationProcessor::ProcessJob(const FGuid& JobGuid, FProxyGeneratio
 	SrcModel->BuildSettings.bRemoveDegenerates = true;
 	SrcModel->BuildSettings.bUseHighPrecisionTangentBasis = false;
 	SrcModel->BuildSettings.bUseFullPrecisionUVs = false;
-	SrcModel->RawMeshBulkData->SaveRawMesh(Data->RawMesh);
+	SrcModel->SaveRawMesh(Data->RawMesh);
 
 	//Assign the proxy material to the static mesh
 	StaticMesh->StaticMaterials.Add(FStaticMaterial(ProxyMaterial));
diff --git a/Engine/Source/Developer/RawMesh/Private/RawMesh.cpp b/Engine/Source/Developer/RawMesh/Private/RawMesh.cpp
--- a/Engine/Source/Developer/RawMesh/Private/RawMesh.cpp
+++ b/Engine/Source/Developer/RawMesh/Private/RawMesh.cpp
@@ -285,4 +285,11 @@ void FRawMeshBulkData::UseHashAsGuid(UObject* Owner)
 	bGuidIsHash = true;
 }
 
+void FRawMeshBulkData::Empty()
+{
+	BulkData.RemoveBulkData();
+	Guid.Invalidate();
+	bGuidIsHash = false;
+}
+
 #endif // #if WITH_EDITORONLY_DATA
diff --git a/Engine/Source/Developer/RawMesh/Public/RawMesh.h b/Engine/Source/Developer/RawMesh/Public/RawMesh.h
--- a/Engine/Source/Developer/RawMesh/Public/RawMesh.h
+++ b/Engine/Source/Developer/RawMesh/Public/RawMesh.h
@@ -120,5 +120,8 @@ public:
 
 	/** Returns true if no bulk data is available for this mesh. */
 	FORCEINLINE bool IsEmpty() const { return BulkData.GetBulkDataSize() == 0; }
+
+	/** Returns true if no bulk data is available for this mesh. */
+	RAWMESH_API void Empty();
 #endif // #if WITH_EDITORONLY_DATA
 };
diff --git a/Engine/Source/Editor/MeshPaint/Private/MeshPaintHelpers.cpp b/Engine/Source/Editor/MeshPaint/Private/MeshPaintHelpers.cpp
--- a/Engine/Source/Editor/MeshPaint/Private/MeshPaintHelpers.cpp
+++ b/Engine/Source/Editor/MeshPaint/Private/MeshPaintHelpers.cpp
@@ -105,7 +105,7 @@ bool MeshPaintHelpers::PropagateColorsToRawMesh(UStaticMesh* StaticMesh, int32 L
 	{
 		// Use the wedge map if it is available as it is lossless.
 		FRawMesh RawMesh;
-		SrcModel.RawMeshBulkData->LoadRawMesh(RawMesh);
+		SrcModel.LoadRawMesh(RawMesh);
 
 		int32 NumWedges = RawMesh.WedgeIndices.Num();
 		if (RenderData.WedgeMap.Num() == NumWedges)
@@ -125,21 +125,21 @@ bool MeshPaintHelpers::PropagateColorsToRawMesh(UStaticMesh* StaticMesh, int32 L
 				}
 				RawMesh.WedgeColors[i] = WedgeColor;
 			}
-			SrcModel.RawMeshBulkData->SaveRawMesh(RawMesh);
+			SrcModel.SaveRawMesh(RawMesh);
 			bPropagatedColors = true;
 		}
 	}
 	else
 	{
 		// If there's no raw mesh data, don't try to do any fixup here
-		if (SrcModel.RawMeshBulkData->IsEmpty() || ComponentLODInfo.OverrideMapBuildData == nullptr)
+		if (SrcModel.IsRawMeshEmpty() || ComponentLODInfo.OverrideMapBuildData == nullptr)
 		{
 			return false;
 		}
 
 		// Fall back to mapping based on position.
 		FRawMesh RawMesh;
-		SrcModel.RawMeshBulkData->LoadRawMesh(RawMesh);
+		SrcModel.LoadRawMesh(RawMesh);
 
 		TArray<FColor> NewVertexColors;
 		FPositionVertexBuffer TempPositionVertexBuffer;
@@ -163,7 +163,7 @@ bool MeshPaintHelpers::PropagateColorsToRawMesh(UStaticMesh* StaticMesh, int32 L
 				int32 Index = RawMesh.WedgeIndices[i];
 				RawMesh.WedgeColors[i] = NewVertexColors[Index];
 			}
-			SrcModel.RawMeshBulkData->SaveRawMesh(RawMesh);
+			SrcModel.SaveRawMesh(RawMesh);
 			bPropagatedColors = true;
 		}
 	}
diff --git a/Engine/Source/Editor/StaticMeshEditor/Private/StaticMeshEditorViewportClient.cpp b/Engine/Source/Editor/StaticMeshEditor/Private/StaticMeshEditorViewportClient.cpp
--- a/Engine/Source/Editor/StaticMeshEditor/Private/StaticMeshEditorViewportClient.cpp
+++ b/Engine/Source/Editor/StaticMeshEditor/Private/StaticMeshEditorViewportClient.cpp
@@ -999,7 +999,7 @@ void FStaticMeshEditorViewportClient::ProcessClick(class FSceneView& InView, cla
 
 				const uint32 LODLevel = FMath::Clamp( StaticMeshComponent->ForcedLodModel - 1, 0, StaticMeshComponent->GetStaticMesh()->GetNumLODs() - 1 );
 				FRawMesh RawMesh;
-				StaticMeshComponent->GetStaticMesh()->SourceModels[LODLevel].RawMeshBulkData->LoadRawMesh(RawMesh);
+				StaticMeshComponent->GetStaticMesh()->SourceModels[LODLevel].LoadRawMesh(RawMesh);
 
 				const int32 RawEdgeCount = RawMesh.WedgeIndices.Num() - 1; 
 				const int32 NumFaces = RawMesh.WedgeIndices.Num() / 3;
diff --git a/Engine/Source/Editor/UnrealEd/Private/EditorExporters.cpp b/Engine/Source/Editor/UnrealEd/Private/EditorExporters.cpp
--- a/Engine/Source/Editor/UnrealEd/Private/EditorExporters.cpp
+++ b/Engine/Source/Editor/UnrealEd/Private/EditorExporters.cpp
@@ -1608,7 +1608,7 @@ bool UStaticMeshExporterOBJ::ExportText(const FExportObjectInnerContext* Context
 		// Currently, we only export LOD 0 of the static mesh. In the future, we could potentially export all available LODs
 		const FStaticMeshLODResources& RenderData = StaticMesh->GetLODForExport(0);
 		FRawMesh RawMesh;
-		StaticMesh->SourceModels[0].RawMeshBulkData->LoadRawMesh(RawMesh);
+		StaticMesh->SourceModels[0].LoadRawMesh(RawMesh);
 
 		uint32 Count = RenderData.GetNumTriangles();
 
diff --git a/Engine/Source/Editor/UnrealEd/Private/Fbx/FbxStaticMeshImport.cpp b/Engine/Source/Editor/UnrealEd/Private/Fbx/FbxStaticMeshImport.cpp
--- a/Engine/Source/Editor/UnrealEd/Private/Fbx/FbxStaticMeshImport.cpp
+++ b/Engine/Source/Editor/UnrealEd/Private/Fbx/FbxStaticMeshImport.cpp
@@ -415,7 +415,7 @@ bool UnFbx::FFbxImporter::BuildStaticMeshFromGeometry(FbxNode* Node, UStaticMesh
 	FbxLayerElement::EMappingMode SmoothingMappingMode(FbxLayerElement::eByEdge);
 	if (SmoothingInfo)
 	{
-		if(MeshDescription != nullptr)
+		if(GetDefault<UEditorExperimentalSettings>()->bUseMeshDescription)
 		{
 			if (SmoothingInfo->GetMappingMode() == FbxLayerElement::eByPolygon)
 			{
@@ -520,7 +520,7 @@ bool UnFbx::FFbxImporter::BuildStaticMeshFromGeometry(FbxNode* Node, UStaticMesh
 	int32 VertexCount = Mesh->GetControlPointsCount();
 	bool OddNegativeScale = IsOddNegativeScale(TotalMatrix);
 	bool bHasNonDegenerateTriangles = false;
-	if(MeshDescription != nullptr)
+	if(GetDefault<UEditorExperimentalSettings>()->bUseMeshDescription)
 	{
 		TVertexAttributeArray<FVector>& VertexPositions = MeshDescription->VertexAttributes().GetAttributes<FVector>(MeshAttribute::Vertex::Position);
 
@@ -547,8 +547,19 @@ bool UnFbx::FFbxImporter::BuildStaticMeshFromGeometry(FbxNode* Node, UStaticMesh
 		TMap<int32, FPolygonGroupID> PolygonGroupMapping;
 
 		// When importing multiple mesh pieces to the same static mesh.  Ensure each mesh piece has the same number of Uv's
-		int32 ExistingUVCount = VertexInstanceUVs.GetNumIndices();
-		
+		int32 ExistingUVCount = 0;
+		for (int32 UVChannelIndex = 0; UVChannelIndex < VertexInstanceUVs.GetNumIndices(); ++UVChannelIndex)
+		{
+			if (VertexInstanceUVs.GetArrayForIndex(UVChannelIndex).Num() > 0)
+			{
+				ExistingUVCount++;
+			}
+			else
+			{
+				break;
+			}
+		}
+
 		int32 NumUVs = FMath::Max(FBXUVs.UniqueUVCount, ExistingUVCount);
 		// At least one UV set must exist.  
 		NumUVs = FMath::Max(1, NumUVs);
@@ -572,18 +583,61 @@ bool UnFbx::FFbxImporter::BuildStaticMeshFromGeometry(FbxNode* Node, UStaticMesh
 				return false;
 			}
 		}
+		
+		Mesh->BeginGetMeshEdgeVertices();
+		TMap<uint32, int32> RemapEdgeID;
+		//Fill the edge array
+		int32 FbxEdgeCount = Mesh->GetMeshEdgeCount();
+		RemapEdgeID.Reserve(FbxEdgeCount*2);
+		for (int32 FbxEdgeIndex = 0; FbxEdgeIndex < FbxEdgeCount; ++FbxEdgeIndex)
+		{
+			int32 EdgeStartVertexIndex = -1;
+			int32 EdgeEndVertexIndex = -1;
+			Mesh->GetMeshEdgeVertices(FbxEdgeIndex, EdgeStartVertexIndex, EdgeEndVertexIndex);
+			FVertexID EdgeVertexStart(EdgeStartVertexIndex + VertexOffset);
+			check(MeshDescription->Vertices().IsValid(EdgeVertexStart));
+			FVertexID EdgeVertexEnd(EdgeEndVertexIndex + VertexOffset);
+			check(MeshDescription->Vertices().IsValid(EdgeVertexEnd));
+			uint32 CompactedKey = ((uint32)EdgeVertexStart.GetValue() << 16) | ((uint32)EdgeVertexEnd.GetValue() & 0xffff);
+			RemapEdgeID.Add(CompactedKey, FbxEdgeIndex);
+			//Add the other edge side
+			CompactedKey = ((uint32)EdgeVertexEnd.GetValue() << 16) | ((uint32)EdgeVertexStart.GetValue() & 0xffff);
+			RemapEdgeID.Add(CompactedKey, FbxEdgeIndex);
+		}
+		//Call this after all GetMeshEdgeIndexForPolygon call this is for optimization purpose.
+		Mesh->EndGetMeshEdgeVertices();
 
 		//Call this before all GetMeshEdgeIndexForPolygon call this is for optimization purpose.
 		Mesh->BeginGetMeshEdgeIndexForPolygon();
+		int32 VertexInstanceCut = 0;
 		//Polygons
 		for (int32 TriangleIndex = 0; TriangleIndex < TriangleCount; TriangleIndex++)
 		{
+			//Verify if the polygon is degenerate, in this case do not add them
+			{
+				float ComparisonThreshold = ImportOptions->bRemoveDegenerates ? SMALL_NUMBER : 0.0f;
+				FVector P[3];
+				for (int32 CornerIndex = 0; CornerIndex < 3; CornerIndex++)
+				{
+					const int32 ControlPointIndex = Mesh->GetPolygonVertex(TriangleIndex, CornerIndex);
+					const FVertexID VertexID(VertexOffset + ControlPointIndex);
+					P[CornerIndex] = VertexPositions[VertexID];
+				}
+				const FVector Normal = ((P[1] - P[2]) ^ (P[0] - P[2])).GetSafeNormal(ComparisonThreshold);
+				//Check for degenerated polygons, avoid NAN
+				if (Normal.IsNearlyZero(ComparisonThreshold))
+				{
+					VertexInstanceCut += 3;
+					continue;
+				}
+			}
+
 			int32 RealTriangleIndex = TriangleOffset + TriangleIndex;
 			FVertexInstanceID CornerInstanceIDs[3];
 			FVertexID CornerVerticesIDs[3];
 			for (int32 CornerIndex = 0; CornerIndex < 3; CornerIndex++)
 			{
-				int32 VertexInstanceIndex = RealTriangleIndex * 3 + CornerIndex;
+				int32 VertexInstanceIndex = RealTriangleIndex * 3 + CornerIndex - VertexInstanceCut;
 				const FVertexInstanceID VertexInstanceID(VertexInstanceIndex);
 				CornerInstanceIDs[CornerIndex] = VertexInstanceID;
 				const int32 ControlPointIndex = Mesh->GetPolygonVertex(TriangleIndex, CornerIndex);
@@ -796,8 +850,19 @@ bool UnFbx::FFbxImporter::BuildStaticMeshFromGeometry(FbxNode* Node, UStaticMesh
 					}
 					ContourPoint.EdgeID = MatchEdgeId;
 					ContourPoint.VertexInstanceID = CornerInstanceIDs[CornerIndices[0]];
-
-					int32 EdgeIndex = Mesh->GetMeshEdgeIndexForPolygon(TriangleIndex, TriangleEdgeNumber);
+					//RawMesh do not have edges, so by ordering the edge with the triangle construction we can ensure back and forth conversion with RawMesh
+					//When raw mesh will be completly remove we can create the edges right after the vertex creation.
+					int32 EdgeIndex = INDEX_NONE;
+					uint32 CompactedKey = ((uint32)EdgeVertexIDs[0].GetValue() << 16) | ((uint32)EdgeVertexIDs[1].GetValue() & 0xffff);
+					if (RemapEdgeID.Contains(CompactedKey))
+					{
+						EdgeIndex = RemapEdgeID[CompactedKey];
+					}
+					else
+					{
+						EdgeIndex = Mesh->GetMeshEdgeIndexForPolygon(TriangleIndex, TriangleEdgeNumber);
+					}
+					
 					EdgeCreaseSharpnesses[MatchEdgeId] = (float)Mesh->GetEdgeCreaseInfo(EdgeIndex);
 					if (!EdgeHardnesses[MatchEdgeId])
 					{
@@ -1602,7 +1667,7 @@ UStaticMesh* UnFbx::FFbxImporter::ImportStaticMeshAsSingle(UObject* InParent, TA
 		}
 	}
 
-	if (MeshDescription == nullptr)
+	if (!GetDefault<UEditorExperimentalSettings>()->bUseMeshDescription)
 	{
 		// Store the new raw mesh.
 		SrcModel.RawMeshBulkData->SaveRawMesh(NewRawMesh);
@@ -1613,7 +1678,7 @@ UStaticMesh* UnFbx::FFbxImporter::ImportStaticMeshAsSingle(UObject* InParent, TA
 	{
 		int32 MaxMaterialIndex = 0;
 		int32 FirstOpenUVChannel = 1;
-		if (MeshDescription == nullptr)
+		if (!GetDefault<UEditorExperimentalSettings>()->bUseMeshDescription)
 		{
 			UE_LOG(LogFbx, Verbose, TEXT("== Initial material list:"));
 			for (int32 MaterialIndex = 0; MaterialIndex < MeshMaterials.Num(); ++MaterialIndex)
@@ -2008,6 +2073,18 @@ void UnFbx::FFbxImporter::PostImportStaticMesh(UStaticMesh* StaticMesh, TArray<F
 	}
 #endif
 
+	if (GIsAutomationTesting)
+	{
+		//Avoid distance field calculation in automation test setting this to false is not suffisant since the condition OR with the CVar
+		//But fbx automation test turn off the CVAR
+		StaticMesh->bGenerateMeshDistanceField = false;
+	}
+
+	static const auto CVarDistanceField = IConsoleManager::Get().FindTConsoleVariableDataInt(TEXT("r.GenerateMeshDistanceFields"));
+	int32 OriginalCVarDistanceFieldValue = CVarDistanceField->GetValueOnGameThread();
+	IConsoleVariable* CVarDistanceFieldInterface = IConsoleManager::Get().FindConsoleVariable(TEXT("r.GenerateMeshDistanceFields"));
+	bool bOriginalGenerateMeshDistanceField = StaticMesh->bGenerateMeshDistanceField;
+
 	//Prebuild the static mesh when we use LodGroup and we want to modify the LodNumber
 	if (!ImportOptions->bImportScene)
 	{
@@ -2026,15 +2103,11 @@ void UnFbx::FFbxImporter::PostImportStaticMesh(UStaticMesh* StaticMesh, TArray<F
 			if (bSpecifiedLodGroup)
 			{
 				//Avoid building the distance field when we prebuild
-				static const auto CVarDistanceField = IConsoleManager::Get().FindTConsoleVariableDataInt(TEXT("r.GenerateMeshDistanceFields"));
-				int32 OriginalCVarDistanceFieldValue = CVarDistanceField->GetValueOnGameThread();
-				IConsoleVariable* CVarDistanceFieldInterface = IConsoleManager::Get().FindConsoleVariable(TEXT("r.GenerateMeshDistanceFields"));
 				if (OriginalCVarDistanceFieldValue != 0 && CVarDistanceFieldInterface)
 				{
 					//Hack we change the distance field user console variable to control the build, but we put back the value after the first build
 					CVarDistanceFieldInterface->SetWithCurrentPriority(0);
 				}
-				bool bOriginalGenerateMeshDistanceField = StaticMesh->bGenerateMeshDistanceField;
 				StaticMesh->bGenerateMeshDistanceField = false;
 
 				StaticMesh->Build(false, &BuildErrors);
diff --git a/Engine/Source/Editor/UnrealEd/Private/GeomFitUtils.cpp b/Engine/Source/Editor/UnrealEd/Private/GeomFitUtils.cpp
--- a/Engine/Source/Editor/UnrealEd/Private/GeomFitUtils.cpp
+++ b/Engine/Source/Editor/UnrealEd/Private/GeomFitUtils.cpp
@@ -22,6 +22,7 @@
 #include "Engine/StaticMesh.h"
 #include "MeshDescription.h"
 #include "MeshAttributes.h"
+#include "Settings/EditorExperimentalSettings.h"
 
 #define LOCAL_EPS (0.01f)
 static void AddVertexIfNotPresent(TArray<FVector> &vertices, FVector &newVertex)
@@ -218,7 +219,7 @@ void ComputeBoundingBox(UStaticMesh* StaticMesh, FVector& Center, FVector& Exten
 	// Calculate bounding Box.
 	
 	FStaticMeshSourceModel& SrcModel = StaticMesh->SourceModels[0];
-	if (StaticMesh->GetOriginalMeshDescription(0) != nullptr)
+	if (GetDefault<UEditorExperimentalSettings>()->bUseMeshDescription)
 	{
 		UMeshDescription* MeshDescription = StaticMesh->GetOriginalMeshDescription(0);
 		FVector unitVec = FVector(1.f);
@@ -227,7 +228,7 @@ void ComputeBoundingBox(UStaticMesh* StaticMesh, FVector& Center, FVector& Exten
 	else
 	{
 		FRawMesh RawMesh;
-		SrcModel.RawMeshBulkData->LoadRawMesh(RawMesh);
+		SrcModel.LoadRawMesh(RawMesh);
 		FVector unitVec = FVector(1.f);
 		CalcBoundingBox(RawMesh, Center, Extents, unitVec);
 	}
@@ -249,14 +250,14 @@ int32 GenerateBoxAsSimpleCollision(UStaticMesh* StaticMesh)
 	
 	FVector unitVec = bs->BuildScale3D;
 	FVector Center, Extents;
-	if (StaticMesh->GetOriginalMeshDescription(0) != nullptr)
+	if (GetDefault<UEditorExperimentalSettings>()->bUseMeshDescription)
 	{
 		MeshDescription = StaticMesh->GetMeshDescription(0);
 		CalcBoundingBox(MeshDescription, Center, Extents, unitVec);
 	}
 	else
 	{
-		SrcModel.RawMeshBulkData->LoadRawMesh(RawMesh);
+		SrcModel.LoadRawMesh(RawMesh);
 		CalcBoundingBox(RawMesh, Center, Extents, unitVec);
 	}
 
@@ -560,7 +561,7 @@ int32 GenerateSphereAsSimpleCollision(UStaticMesh* StaticMesh)
 	FVector unitVec = bs->BuildScale3D;
 
 	// Calculate bounding sphere.
-	if (StaticMesh->GetOriginalMeshDescription(0) != nullptr)
+	if (GetDefault<UEditorExperimentalSettings>()->bUseMeshDescription)
 	{
 		UMeshDescription* MeshDescription = StaticMesh->GetMeshDescription(0);
 		CalcBoundingSphere(MeshDescription, bSphere, unitVec);
@@ -570,7 +571,7 @@ int32 GenerateSphereAsSimpleCollision(UStaticMesh* StaticMesh)
 	{
 		FRawMesh RawMesh;
 		FStaticMeshSourceModel& SrcModel = StaticMesh->SourceModels[0];
-		SrcModel.RawMeshBulkData->LoadRawMesh(RawMesh);
+		SrcModel.LoadRawMesh(RawMesh);
 		CalcBoundingSphere(RawMesh, bSphere, unitVec);
 		CalcBoundingSphere2(RawMesh, bSphere2, unitVec);
 	}
@@ -796,7 +797,7 @@ int32 GenerateSphylAsSimpleCollision(UStaticMesh* StaticMesh)
 	FVector unitVec = bs->BuildScale3D;
 
 	// Calculate bounding box.
-	if (StaticMesh->GetOriginalMeshDescription(0) != nullptr)
+	if (GetDefault<UEditorExperimentalSettings>()->bUseMeshDescription)
 	{
 		UMeshDescription* MeshDescription = StaticMesh->GetMeshDescription(0);
 		CalcBoundingSphyl(MeshDescription, sphere, length, rotation, unitVec);
@@ -805,7 +806,7 @@ int32 GenerateSphylAsSimpleCollision(UStaticMesh* StaticMesh)
 	{
 		FRawMesh RawMesh;
 		FStaticMeshSourceModel& SrcModel = StaticMesh->SourceModels[0];
-		SrcModel.RawMeshBulkData->LoadRawMesh(RawMesh);
+		SrcModel.LoadRawMesh(RawMesh);
 		CalcBoundingSphyl(RawMesh, sphere, length, rotation, unitVec);
 	}
 
diff --git a/Engine/Source/Editor/UnrealEd/Private/StaticMeshEdit.cpp b/Engine/Source/Editor/UnrealEd/Private/StaticMeshEdit.cpp
--- a/Engine/Source/Editor/UnrealEd/Private/StaticMeshEdit.cpp
+++ b/Engine/Source/Editor/UnrealEd/Private/StaticMeshEdit.cpp
@@ -583,7 +583,7 @@ UStaticMesh* CreateStaticMesh(struct FRawMesh& RawMesh,TArray<FStaticMaterial>&
 
 	// Add one LOD for the base mesh
 	FStaticMeshSourceModel* SrcModel = new(StaticMesh->SourceModels) FStaticMeshSourceModel();
-	SrcModel->RawMeshBulkData->SaveRawMesh(RawMesh);
+	SrcModel->SaveRawMesh(RawMesh);
 	StaticMesh->StaticMaterials = Materials;
 
 	int32 NumSections = StaticMesh->StaticMaterials.Num();
diff --git a/Engine/Source/Editor/WorldBrowser/Private/Tiles/WorldTileCollectionModel.cpp b/Engine/Source/Editor/WorldBrowser/Private/Tiles/WorldTileCollectionModel.cpp
--- a/Engine/Source/Editor/WorldBrowser/Private/Tiles/WorldTileCollectionModel.cpp
+++ b/Engine/Source/Editor/WorldBrowser/Private/Tiles/WorldTileCollectionModel.cpp
@@ -2149,7 +2149,7 @@ bool FWorldTileCollectionModel::GenerateLODLevels(FLevelModelList InLevelList, i
 				SrcModel->BuildSettings.bRemoveDegenerates = false;
 				SrcModel->BuildSettings.bUseHighPrecisionTangentBasis = false;
 				SrcModel->BuildSettings.bUseFullPrecisionUVs = false;
-				SrcModel->RawMeshBulkData->SaveRawMesh(LandscapeRawMesh);
+				SrcModel->SaveRawMesh(LandscapeRawMesh);
 
 				//Assign the proxy material to the static mesh
 				StaticMesh->StaticMaterials.Add(FStaticMaterial(StaticLandscapeMaterial));
diff --git a/Engine/Source/Runtime/Engine/Classes/Engine/StaticMesh.h b/Engine/Source/Runtime/Engine/Classes/Engine/StaticMesh.h
--- a/Engine/Source/Runtime/Engine/Classes/Engine/StaticMesh.h
+++ b/Engine/Source/Runtime/Engine/Classes/Engine/StaticMesh.h
@@ -154,10 +154,22 @@ struct FStaticMeshSourceModel
 	/** Imported raw mesh data. Optional for all but the first LOD. */
 	class FRawMeshBulkData* RawMeshBulkData;
 
-	/* Original Imported meshdescription data. Optional for all but the first LOD. */
+	/*
+	 * Accessor to Load and save the raw mesh or the mesh description depending on the editor settings.
+	 * Temporary until we deprecate the RawMesh. It use UEditorExperimentalSettings::bUseMeshDescription to know which structure we use to save/load the data.
+	 * The SaveRawMesh make sure there is only one valid structure data (OriginalMeshDescription or RawMesh) not both. 
+	 */
+	ENGINE_API bool IsRawMeshEmpty() const;
+	ENGINE_API void LoadRawMesh(struct FRawMesh& OutRawMesh) const;
+	ENGINE_API void SaveRawMesh(struct FRawMesh& InRawMesh);
+
+#endif // #if WITH_EDITOR
+
+#if WITH_EDITORONLY_DATA
+	/* Original Imported mesh description data. Optional for all but the first LOD. Autogenerate LOD do not have original mesh description*/
 	UPROPERTY()
 	class UMeshDescription* OriginalMeshDescription;
-#endif // #if WITH_EDITOR
+#endif
 
 	/** Settings applied when building the mesh. */
 	UPROPERTY(EditAnywhere, Category=BuildSettings)
@@ -358,11 +370,9 @@ struct FStaticMaterial
 	UPROPERTY(EditAnywhere, BlueprintReadOnly, Category = StaticMesh)
 	FName MaterialSlotName;
 
-#if WITH_EDITORONLY_DATA
 	/*This name should be use when we re-import a skeletal mesh so we can order the Materials array like it should be*/
 	UPROPERTY(VisibleAnywhere, Category = StaticMesh)
 	FName ImportedMaterialSlotName;
-#endif //WITH_EDITORONLY_DATA
 
 	/** Data used for texture streaming relative to each UV channels. */
 	UPROPERTY(VisibleAnywhere, BlueprintReadOnly, Category = StaticMesh)
@@ -629,16 +639,26 @@ public:
 	UPROPERTY(Instanced, VisibleAnywhere, Category = EditableMesh)
 	class UObject* EditableMesh;
 
-	/** The MeshDescription use to build the render data. */
+#if WITH_EDITORONLY_DATA
+	/**
+	 * The MeshDescription use to build the render data.
+	 * When building the render data we copy the original mesh description and modify it to
+	 * respect the build options (tangent, UVs, reduce, ...) we then keep the copy in this array in case we need it
+	 */
 	UPROPERTY(VisibleAnywhere, Category = MeshDescription)
 	TArray<class UMeshDescription*> MeshDescriptions;
-
+	/**
+	 * Accessors for the mesh description use to build the render data
+	 */
 	ENGINE_API class UMeshDescription* GetMeshDescription(int32 LodIndex=0) const;
 	ENGINE_API void SetMeshDescription(int32 LodIndex, class UMeshDescription* InMeshDescription);
 	ENGINE_API int32 GetMeshDescriptionCount() const;
 
-#if WITH_EDITORONLY_DATA
-	ENGINE_API class UMeshDescription* GetOriginalMeshDescription(int32 LodIndex = 0) const;
+	/**
+	 * Accessors for the original mesh description imported data
+	 * The original import data is necessary to start from the full data when applying build options.
+	 */
+	ENGINE_API class UMeshDescription* GetOriginalMeshDescription(int32 LodIndex = 0);
 	ENGINE_API void SetOriginalMeshDescription(int32 LodIndex, class UMeshDescription* MeshDescription);
 	ENGINE_API void ClearOriginalMeshDescription(int32 LodIndex);
 #endif
diff --git a/Engine/Source/Runtime/Engine/Engine.Build.cs b/Engine/Source/Runtime/Engine/Engine.Build.cs
--- a/Engine/Source/Runtime/Engine/Engine.Build.cs
+++ b/Engine/Source/Runtime/Engine/Engine.Build.cs
@@ -79,7 +79,7 @@ public class Engine : ModuleRules
 				"DatabaseSupport",
 				"PacketHandler",
 				"HardwareSurvey",
-				"MeshDescription"
+				"MeshDescription",
 			}
 		);
 
@@ -116,6 +116,7 @@ public class Engine : ModuleRules
 		{
 			// for now we depend on this
 			PrivateDependencyModuleNames.Add("RawMesh");
+			PrivateDependencyModuleNames.Add("MeshBuilder");
 		}
 
 		bool bVariadicTemplatesSupported = true;
diff --git a/Engine/Source/Runtime/Engine/Private/StaticMesh.cpp b/Engine/Source/Runtime/Engine/Private/StaticMesh.cpp
--- a/Engine/Source/Runtime/Engine/Private/StaticMesh.cpp
+++ b/Engine/Source/Runtime/Engine/Private/StaticMesh.cpp
@@ -39,7 +39,10 @@
 
 #if WITH_EDITOR
 #include "RawMesh.h"
+#include "Settings/EditorExperimentalSettings.h"
 #include "MeshDescription.h"
+#include "MeshAttributes.h"
+#include "MeshBuilder.h"
 #include "MeshUtilities.h"
 #include "DerivedDataCacheInterface.h"
 #include "IMeshBuilderModule.h"
@@ -1151,9 +1154,10 @@ static FString BuildStaticMeshDerivedDataKey(UStaticMesh* Mesh, const FStaticMes
 	for (int32 LODIndex = 0; LODIndex < NumLODs; ++LODIndex)
 	{
 		FStaticMeshSourceModel& SrcModel = Mesh->SourceModels[LODIndex];
-		if (Mesh->GetOriginalMeshDescription(0) != nullptr)
+		if (GetDefault<UEditorExperimentalSettings>()->bUseMeshDescription)
 		{
-			KeySuffix += (Mesh->GetOriginalMeshDescription(LODIndex) != nullptr) ? Mesh->GetOriginalMeshDescription(LODIndex)->GetIdString() : TEXT("NoOriginalMeshDescriptionForLod") + FString::FromInt(LODIndex);
+			UMeshDescription* MeshDescription = Mesh->GetOriginalMeshDescription(LODIndex);
+			KeySuffix += (MeshDescription != nullptr) ? MeshDescription->GetIdString() : TEXT("NoOriginalMeshDescriptionForLod") + FString::FromInt(LODIndex);
 		}
 		else
 		{
@@ -1302,10 +1306,23 @@ void FStaticMeshRenderData::Cache(UStaticMesh* Owner, const FStaticMeshLODSettin
 			Args.Add(TEXT("StaticMeshName"), FText::FromString( Owner->GetName() ) );
 			FStaticMeshStatusMessageContext StatusContext( FText::Format( NSLOCTEXT("Engine", "BuildingStaticMeshStatus", "Building static mesh {StaticMeshName}..."), Args ) );
 
-			if (Owner->GetOriginalMeshDescription(0) != nullptr)
+			bool bUseMeshDescription = false;
+			if (GetDefault<UEditorExperimentalSettings>()->bUseMeshDescription)
+			{
+				if (Owner->GetOriginalMeshDescription(0) != nullptr)
+				{
+					bUseMeshDescription = true;
+				}
+			}
+
+			if (bUseMeshDescription)
 			{
 				IMeshBuilderModule& MeshBuilderModule = FModuleManager::Get().LoadModuleChecked<IMeshBuilderModule>(TEXT("MeshBuilder"));
-				MeshBuilderModule.BuildMesh(Owner, LODGroup);
+				if (!MeshBuilderModule.BuildMesh(Owner, LODGroup))
+				{
+					UE_LOG(LogStaticMesh, Error, TEXT("Failed to build static mesh. See previous line(s) for details."));
+					return;
+				}
 			}
 			else
 			{
@@ -2060,6 +2077,51 @@ FStaticMeshSourceModel::~FStaticMeshSourceModel()
 }
 
 #if WITH_EDITOR
+bool FStaticMeshSourceModel::IsRawMeshEmpty() const
+{
+	return (RawMeshBulkData == nullptr || RawMeshBulkData->IsEmpty() && OriginalMeshDescription == nullptr);
+}
+
+void FStaticMeshSourceModel::LoadRawMesh(FRawMesh& OutRawMesh) const
+{
+	bool bUseMeshDescription = GetDefault<UEditorExperimentalSettings>()->bUseMeshDescription;
+	if (bUseMeshDescription && OriginalMeshDescription != nullptr)
+	{
+		FMeshDescriptionOperations::ConverToRawMesh(OriginalMeshDescription, OutRawMesh);
+	}
+	else
+	{
+		if (RawMeshBulkData->IsEmpty() && OriginalMeshDescription != nullptr)
+		{
+			FMeshDescriptionOperations::ConverToRawMesh(OriginalMeshDescription, OutRawMesh);
+		}
+		else
+		{
+			RawMeshBulkData->LoadRawMesh(OutRawMesh);
+		}
+	}
+	
+}
+
+void FStaticMeshSourceModel::SaveRawMesh(FRawMesh& InRawMesh)
+{
+	if (!InRawMesh.IsValid())
+	{
+		return;
+	}
+	bool bUseMeshDescription = GetDefault<UEditorExperimentalSettings>()->bUseMeshDescription;
+	if (bUseMeshDescription && OriginalMeshDescription != nullptr)
+	{
+		FMeshDescriptionOperations::ConverFromRawMesh(InRawMesh, OriginalMeshDescription);
+		RawMeshBulkData->Empty();
+	}
+	else
+	{
+		RawMeshBulkData->SaveRawMesh(InRawMesh);
+		OriginalMeshDescription = nullptr;
+	}
+}
+
 void FStaticMeshSourceModel::SerializeBulkData(FArchive& Ar, UObject* Owner)
 {
 	check(RawMeshBulkData != NULL);
@@ -2217,6 +2279,7 @@ static FStaticMeshRenderData& GetPlatformStaticMeshRenderData(UStaticMesh* Mesh,
 	return *PlatformRenderData;
 }
 
+#if WITH_EDITORONLY_DATA
 
 UMeshDescription* UStaticMesh::GetMeshDescription(int32 LodIndex) const
 {
@@ -2243,13 +2306,39 @@ int32 UStaticMesh::GetMeshDescriptionCount() const
 	return MeshDescriptions.Num();
 }
 
-#if WITH_EDITORONLY_DATA
-
-UMeshDescription* UStaticMesh::GetOriginalMeshDescription(int32 LodIndex) const
+UMeshDescription* UStaticMesh::GetOriginalMeshDescription(int32 LodIndex)
 {
 	if (SourceModels.IsValidIndex(LodIndex))
 	{
-		return SourceModels[LodIndex].OriginalMeshDescription;
+		if (GetDefault<UEditorExperimentalSettings>()->bUseMeshDescription)
+		{
+			if (SourceModels[LodIndex].OriginalMeshDescription == nullptr && !SourceModels[LodIndex].RawMeshBulkData->IsEmpty())
+			{
+				FRawMesh TempRawMesh;
+				SourceModels[LodIndex].RawMeshBulkData->LoadRawMesh(TempRawMesh);
+				//The original mesh description is null we must create one
+				SourceModels[LodIndex].OriginalMeshDescription = NewObject<UMeshDescription>(this, NAME_None);
+				//Convert RawMesh to meshdescription
+				FMeshDescriptionOperations::ConverFromRawMesh(TempRawMesh, SourceModels[LodIndex].OriginalMeshDescription);
+			}
+			return SourceModels[LodIndex].OriginalMeshDescription;
+		}
+		else
+		{
+			if (!SourceModels[LodIndex].RawMeshBulkData->IsEmpty())
+			{
+				FRawMesh TempRawMesh;
+				SourceModels[LodIndex].RawMeshBulkData->LoadRawMesh(TempRawMesh);
+				if (SourceModels[LodIndex].OriginalMeshDescription == nullptr)
+				{
+					//The original mesh description is null we must create one
+					SourceModels[LodIndex].OriginalMeshDescription = NewObject<UMeshDescription>(this, NAME_None);
+				}
+				//Convert RawMesh to meshdescription
+				FMeshDescriptionOperations::ConverFromRawMesh(TempRawMesh, SourceModels[LodIndex].OriginalMeshDescription);
+			}
+			return SourceModels[LodIndex].OriginalMeshDescription;
+		}
 	}
 	return nullptr;
 }
@@ -2258,7 +2347,27 @@ void UStaticMesh::SetOriginalMeshDescription(int32 LodIndex, class UMeshDescript
 {
 	if (SourceModels.IsValidIndex(LodIndex))
 	{
-		SourceModels[LodIndex].OriginalMeshDescription = MeshDescription;
+		if (!GetDefault<UEditorExperimentalSettings>()->bUseMeshDescription)
+		{
+			if (MeshDescription != nullptr)
+			{
+				FRawMesh TempRawMesh;
+				//Convert RawMesh to meshdescription
+				FMeshDescriptionOperations::ConverToRawMesh(MeshDescription, TempRawMesh);
+				SourceModels[LodIndex].RawMeshBulkData->SaveRawMesh(TempRawMesh);
+			}
+			else
+			{
+				//We have a null mesh description so empty the raw mesh
+				SourceModels[LodIndex].RawMeshBulkData->Empty();
+			}
+			SourceModels[LodIndex].OriginalMeshDescription = nullptr;
+		}
+		else
+		{
+			SourceModels[LodIndex].OriginalMeshDescription = MeshDescription;
+			SourceModels[LodIndex].RawMeshBulkData->Empty();
+		}
 	}
 }
 
@@ -2580,7 +2689,7 @@ void UStaticMesh::Serialize(FArchive& Ar)
 			FStaticMeshSourceModel& SourceModel = SourceModels[LODIndex];
 
 			FRawMesh RawMesh;
-			SourceModel.RawMeshBulkData->LoadRawMesh(RawMesh);
+			SourceModel.LoadRawMesh(RawMesh);
 
 			// Billboard LOD is made up out of quads so check for this
 			bool bQuadVertices = ((RawMesh.VertexPositions.Num() % 4) == 0);
@@ -2644,10 +2753,9 @@ void UStaticMesh::PostLoad()
 
 			for (int32 i = 0; i < SourceModels.Num(); ++i)
 			{
-				FRawMeshBulkData* RawMeshBulkData = SourceModels[i].RawMeshBulkData;
-				if (RawMeshBulkData)
+				if (!SourceModels[i].IsRawMeshEmpty())
 				{
-					RawMeshBulkData->LoadRawMesh(TempRawMesh);
+					SourceModels[i].LoadRawMesh(TempRawMesh);
 					TotalIndexCount += TempRawMesh.WedgeIndices.Num();
 				}
 			}
@@ -3107,27 +3215,50 @@ void UStaticMesh::GetVertexColorData(TMap<FVector, FColor>& VertexColorData)
 	// What LOD to get vertex colors from.  
 	// Currently mesh painting only allows for painting on the first lod.
 	const uint32 PaintingMeshLODIndex = 0;
-	if (SourceModels.IsValidIndex(PaintingMeshLODIndex)
-		&& SourceModels[PaintingMeshLODIndex].RawMeshBulkData->IsEmpty() == false)
+	if (SourceModels.IsValidIndex(PaintingMeshLODIndex))
 	{
-		// Extract the raw mesh.
-		FRawMesh Mesh;
-		SourceModels[PaintingMeshLODIndex].RawMeshBulkData->LoadRawMesh(Mesh);
-
-		// Nothing to copy if there are no colors stored.
-		if (Mesh.WedgeColors.Num() != 0 && Mesh.WedgeColors.Num() == Mesh.WedgeIndices.Num())
+		if (GetDefault<UEditorExperimentalSettings>()->bUseMeshDescription)
 		{
-			// Build a mapping of vertex positions to vertex colors.
-			for (int32 WedgeIndex = 0; WedgeIndex < Mesh.WedgeIndices.Num(); ++WedgeIndex)
+			UMeshDescription* PaintLODMeshDescription = GetOriginalMeshDescription(PaintingMeshLODIndex);
+			if (PaintLODMeshDescription != nullptr)
 			{
-				FVector Position = Mesh.VertexPositions[Mesh.WedgeIndices[WedgeIndex]];
-				FColor Color = Mesh.WedgeColors[WedgeIndex];
-				if (!VertexColorData.Contains(Position))
+				const TVertexAttributeArray<FVector>& VertexPositions = PaintLODMeshDescription->VertexAttributes().GetAttributes<FVector>(MeshAttribute::Vertex::Position);
+				const TVertexInstanceAttributeArray<FVector4>& VertexInstanceColors = PaintLODMeshDescription->VertexInstanceAttributes().GetAttributes<FVector4>(MeshAttribute::VertexInstance::Color);
+				if (VertexInstanceColors.Num() > 0)
 				{
-					VertexColorData.Add(Position, Color);
+					// Build a mapping of vertex positions to vertex colors.
+					for (const FVertexInstanceID& VertexInstanceID : PaintLODMeshDescription->VertexInstances().GetElementIDs())
+					{
+						FVector Position = VertexPositions[PaintLODMeshDescription->GetVertexInstanceVertex(VertexInstanceID)];
+						FColor Color = FLinearColor(VertexInstanceColors[VertexInstanceID]).ToFColor(true);
+						if (!VertexColorData.Contains(Position))
+						{
+							VertexColorData.Add(Position, Color);
+						}
+					}
 				}
 			}
-		}	
+		}
+		else if (SourceModels[PaintingMeshLODIndex].IsRawMeshEmpty() == false)
+		{
+			// Extract the raw mesh.
+			FRawMesh Mesh;
+			SourceModels[PaintingMeshLODIndex].LoadRawMesh(Mesh);
+			// Nothing to copy if there are no colors stored.
+			if (Mesh.WedgeColors.Num() != 0 && Mesh.WedgeColors.Num() == Mesh.WedgeIndices.Num())
+			{
+				// Build a mapping of vertex positions to vertex colors.
+				for (int32 WedgeIndex = 0; WedgeIndex < Mesh.WedgeIndices.Num(); ++WedgeIndex)
+				{
+					FVector Position = Mesh.VertexPositions[Mesh.WedgeIndices[WedgeIndex]];
+					FColor Color = Mesh.WedgeColors[WedgeIndex];
+					if (!VertexColorData.Contains(Position))
+					{
+						VertexColorData.Add(Position, Color);
+					}
+				}
+			}
+		}
 	}
 #endif // #if WITH_EDITORONLY_DATA
 }
@@ -3145,37 +3276,68 @@ void UStaticMesh::SetVertexColorData(const TMap<FVector, FColor>& VertexColorDat
 	// What LOD to get vertex colors from.  
 	// Currently mesh painting only allows for painting on the first lod.
 	const uint32 PaintingMeshLODIndex = 0;
-	if (SourceModels.IsValidIndex(PaintingMeshLODIndex)
-		&& SourceModels[PaintingMeshLODIndex].RawMeshBulkData->IsEmpty() == false)
+	if (SourceModels.IsValidIndex(PaintingMeshLODIndex))
 	{
-		// Extract the raw mesh.
-		FRawMesh Mesh;
-		SourceModels[PaintingMeshLODIndex].RawMeshBulkData->LoadRawMesh(Mesh);
-
-		// Reserve space for the new vertex colors.
-		if (Mesh.WedgeColors.Num() == 0 || Mesh.WedgeColors.Num() != Mesh.WedgeIndices.Num())
+		if (GetDefault<UEditorExperimentalSettings>()->bUseMeshDescription)
 		{
-			Mesh.WedgeColors.Empty(Mesh.WedgeIndices.Num());
-			Mesh.WedgeColors.AddUninitialized(Mesh.WedgeIndices.Num());
+			UMeshDescription* PaintLODMeshDescription = GetOriginalMeshDescription(PaintingMeshLODIndex);
+			if (PaintLODMeshDescription != nullptr)
+			{
+				const TVertexAttributeArray<FVector>& VertexPositions = PaintLODMeshDescription->VertexAttributes().GetAttributes<FVector>(MeshAttribute::Vertex::Position);
+				//Empty vertex color attributes
+				PaintLODMeshDescription->VertexInstanceAttributes().GetAttributesSet<FVector4>(MeshAttribute::VertexInstance::Color).SetNumIndices(0);
+				PaintLODMeshDescription->VertexInstanceAttributes().GetAttributesSet<FVector4>(MeshAttribute::VertexInstance::Color).SetNumIndices(1);
+				TVertexInstanceAttributeArray<FVector4>& VertexInstanceColors = PaintLODMeshDescription->VertexInstanceAttributes().GetAttributes<FVector4>(MeshAttribute::VertexInstance::Color);
+				if (PaintLODMeshDescription->VertexInstances().Num() > 0)
+				{
+					// Build a mapping of vertex positions to vertex colors.
+					for (const FVertexInstanceID& VertexInstanceID : PaintLODMeshDescription->VertexInstances().GetElementIDs())
+					{
+						FVector Position = VertexPositions[PaintLODMeshDescription->GetVertexInstanceVertex(VertexInstanceID)];
+						const FColor* Color = VertexColorData.Find(Position);
+						if (Color)
+						{
+							VertexInstanceColors[VertexInstanceID] = FLinearColor(*Color);
+						}
+						else
+						{
+							VertexInstanceColors[VertexInstanceID] = FLinearColor::White;
+						}
+					}
+				}
+			}
 		}
-
-		// Build a mapping of vertex positions to vertex colors.
-		for (int32 WedgeIndex = 0; WedgeIndex < Mesh.WedgeIndices.Num(); ++WedgeIndex)
+		else if (SourceModels[PaintingMeshLODIndex].IsRawMeshEmpty() == false)
 		{
-			FVector Position = Mesh.VertexPositions[Mesh.WedgeIndices[WedgeIndex]];
-			const FColor* Color = VertexColorData.Find(Position);
-			if (Color)
+			// Extract the raw mesh.
+			FRawMesh Mesh;
+			SourceModels[PaintingMeshLODIndex].LoadRawMesh(Mesh);
+
+			// Reserve space for the new vertex colors.
+			if (Mesh.WedgeColors.Num() == 0 || Mesh.WedgeColors.Num() != Mesh.WedgeIndices.Num())
 			{
-				Mesh.WedgeColors[WedgeIndex] = *Color;
+				Mesh.WedgeColors.Empty(Mesh.WedgeIndices.Num());
+				Mesh.WedgeColors.AddUninitialized(Mesh.WedgeIndices.Num());
 			}
-			else
+
+			// Build a mapping of vertex positions to vertex colors.
+			for (int32 WedgeIndex = 0; WedgeIndex < Mesh.WedgeIndices.Num(); ++WedgeIndex)
 			{
-				Mesh.WedgeColors[WedgeIndex] = FColor(255,255,255,255);
+				FVector Position = Mesh.VertexPositions[Mesh.WedgeIndices[WedgeIndex]];
+				const FColor* Color = VertexColorData.Find(Position);
+				if (Color)
+				{
+					Mesh.WedgeColors[WedgeIndex] = *Color;
+				}
+				else
+				{
+					Mesh.WedgeColors[WedgeIndex] = FColor(255, 255, 255, 255);
+				}
 			}
-		}
 
-		// Save the new raw mesh.
-		SourceModels[PaintingMeshLODIndex].RawMeshBulkData->SaveRawMesh(Mesh);
+			// Save the new raw mesh.
+			SourceModels[PaintingMeshLODIndex].SaveRawMesh(Mesh);
+		}
 	}
 	// TODO_STATICMESH: Build?
 #endif // #if WITH_EDITOR
@@ -3188,16 +3350,27 @@ ENGINE_API void UStaticMesh::RemoveVertexColors()
 
 	for (FStaticMeshSourceModel& SourceModel : SourceModels)
 	{
-		if (SourceModel.RawMeshBulkData && !SourceModel.RawMeshBulkData->IsEmpty())
+		if (GetDefault<UEditorExperimentalSettings>()->bUseMeshDescription)
+		{
+			//Empty the attribute indices
+			if (SourceModel.OriginalMeshDescription != nullptr)
+			{
+				//Remove attributes content
+				SourceModel.OriginalMeshDescription->VertexInstanceAttributes().GetAttributesSet<FVector4>(MeshAttribute::VertexInstance::Color).SetNumIndices(0);
+				//Add back one indice so we are ready to add again some paint color attribute
+				SourceModel.OriginalMeshDescription->VertexInstanceAttributes().GetAttributesSet<FVector4>(MeshAttribute::VertexInstance::Color).SetNumIndices(1);
+			}
+		}
+		else if (!SourceModel.IsRawMeshEmpty())
 		{
 			FRawMesh RawMesh;
-			SourceModel.RawMeshBulkData->LoadRawMesh(RawMesh);
+			SourceModel.LoadRawMesh(RawMesh);
 
 			if (RawMesh.WedgeColors.Num() > 0)
 			{
 				RawMesh.WedgeColors.Empty();
 
-				SourceModel.RawMeshBulkData->SaveRawMesh(RawMesh);
+				SourceModel.SaveRawMesh(RawMesh);
 
 				bRemovedVertexColors = true;
 			}
diff --git a/Engine/Source/Runtime/MeshBuilder/MeshBuilder.Build.cs b/Engine/Source/Runtime/MeshBuilder/MeshBuilder.Build.cs
--- a/Engine/Source/Runtime/MeshBuilder/MeshBuilder.Build.cs
+++ b/Engine/Source/Runtime/MeshBuilder/MeshBuilder.Build.cs
@@ -18,7 +18,8 @@ namespace UnrealBuildTool.Rules
                     "RenderCore",
                     "MeshDescription",
                     "MeshReductionInterface",
-                    "RenderCore"
+                    "RenderCore",
+                    "RawMesh"
                 }
 			);
 
diff --git a/Engine/Source/Runtime/MeshBuilder/Private/BuildOptimizationHelper.h b/Engine/Source/Runtime/MeshBuilder/Private/BuildOptimizationHelper.h
--- a/Engine/Source/Runtime/MeshBuilder/Private/BuildOptimizationHelper.h
+++ b/Engine/Source/Runtime/MeshBuilder/Private/BuildOptimizationHelper.h
@@ -344,30 +344,9 @@ public:
 
 	static void CacheOptimizeVertexAndIndexBuffer(
 		TArray<FStaticMeshBuildVertex>& Vertices,
-		TArray<TArray<uint32> >& PerSectionIndices,
-		TArray<int32>& WedgeMap
+		TArray<TArray<uint32> >& PerSectionIndices
 	)
 	{
-		//////////////////////////////////////////////////////////////////////////
-
-		// TODO
-		// this code is commented because it make the indexes incoherent
-		// Is it something we must do? Not sure!
-		// We should investigate why.
-		// Currently we at least optimize the section indexes, but we do not reorder the vertex array
-
-		//////////////////////////////////////////////////////////////////////////
-
-
-		// Copy the vertices since we will be reordering them
-		//TArray<FStaticMeshBuildVertex> OriginalVertices = Vertices;
-
-		// Initialize a cache that stores which indices have been assigned
-		//TArray<int32> IndexCache;
-		//IndexCache.AddUninitialized(Vertices.Num());
-		//FMemory::Memset(IndexCache.GetData(), INDEX_NONE, IndexCache.Num() * IndexCache.GetTypeSize());
-		//int32 NextAvailableIndex = 0;
-
 		// Iterate through the section index buffers, 
 		// Optimizing index order for the post transform cache (minimizes the number of vertices transformed), 
 		// And vertex order for the pre transform cache (minimizes the amount of vertex data fetched by the GPU).
@@ -379,44 +358,7 @@ public:
 			{
 				// Optimize the index buffer for the post transform cache with.
 				CacheOptimizeIndexBuffer(Indices);
-/*
-				// Copy the index buffer since we will be reordering it
-				TArray<uint32> OriginalIndices = Indices;
-
-				// Go through the indices and assign them new values that are coherent where possible
-				for (int32 Index = 0; Index < Indices.Num(); Index++)
-				{
-					const int32 CachedIndex = IndexCache[OriginalIndices[Index]];
-
-					if (CachedIndex == INDEX_NONE)
-					{
-						// No new index has been allocated for this existing index, assign a new one
-						Indices[Index] = NextAvailableIndex;
-						// Mark what this index has been assigned to
-						IndexCache[OriginalIndices[Index]] = NextAvailableIndex;
-						NextAvailableIndex++;
-					}
-					else
-					{
-						// Reuse an existing index assignment
-						Indices[Index] = CachedIndex;
-					}
-					// Reorder the vertices based on the new index assignment
-					Vertices[Indices[Index]] = OriginalVertices[OriginalIndices[Index]];
-
-				}
-*/
-			}
-		}
-/*
-		for (int32 i = 0; i < WedgeMap.Num(); i++)
-		{
-			int32 MappedIndex = WedgeMap[i];
-			if (MappedIndex != INDEX_NONE)
-			{
-				WedgeMap[i] = IndexCache[MappedIndex];
 			}
 		}
-*/
 	}
 };
\ No newline at end of file
diff --git a/Engine/Source/Runtime/MeshBuilder/Private/MeshBuilder.cpp b/Engine/Source/Runtime/MeshBuilder/Private/MeshBuilder.cpp
--- a/Engine/Source/Runtime/MeshBuilder/Private/MeshBuilder.cpp
+++ b/Engine/Source/Runtime/MeshBuilder/Private/MeshBuilder.cpp
@@ -1,8 +1,24 @@
 // Copyright 1998-2017 Epic Games, Inc. All Rights Reserved.
 
 #include "MeshBuilder.h"
+#include "MeshDescriptionHelper.h"
+#include "MeshDescription.h"
+#include "RawMesh.h"
+
 
 FMeshBuilder::FMeshBuilder()
 {
 
-}
\ No newline at end of file
+}
+
+/** Convert this mesh description into the old FRawMesh format*/
+void FMeshDescriptionOperations::ConverToRawMesh(const UMeshDescription* SourceMeshDescription, FRawMesh &DestinationRawMesh)
+{
+	FMeshDescriptionHelper::ConverToRawMesh(SourceMeshDescription, DestinationRawMesh);
+}
+
+/** Convert old FRawMesh format to MeshDescription*/
+void FMeshDescriptionOperations::ConverFromRawMesh(const FRawMesh &SourceRawMesh, UMeshDescription* DestinationMeshDescription)
+{
+	FMeshDescriptionHelper::ConverFromRawMesh(SourceRawMesh, DestinationMeshDescription);
+}
diff --git a/Engine/Source/Runtime/MeshBuilder/Private/MeshDescriptionHelper.cpp b/Engine/Source/Runtime/MeshBuilder/Private/MeshDescriptionHelper.cpp
--- a/Engine/Source/Runtime/MeshBuilder/Private/MeshDescriptionHelper.cpp
+++ b/Engine/Source/Runtime/MeshBuilder/Private/MeshDescriptionHelper.cpp
@@ -18,7 +18,8 @@
 #include "LayoutUV.h"
 #include "IMeshReductionInterfaces.h"
 #include "IMeshReductionManagerModule.h"
-#include "BuildStatisticManager.h"
+#include "Materials/Material.h"
+#include "RawMesh.h"
 
 //Enable all check
 //#define ENABLE_NTB_CHECK
@@ -42,7 +43,6 @@ FMeshDescriptionHelper::FMeshDescriptionHelper(FMeshBuildSettings* InBuildSettin
 
 UMeshDescription* FMeshDescriptionHelper::GetRenderMeshDescription(UObject* Owner)
 {
-	BuildStatisticManager::FBuildStatisticScope StatScope(TEXT("GetRenderMeshDescription took"));
 	UStaticMesh* StaticMesh = Cast<UStaticMesh>(Owner);
 	check(StaticMesh);
 	//Use the build settings to create the RenderMeshDescription
@@ -120,10 +120,26 @@ UMeshDescription* FMeshDescriptionHelper::GetRenderMeshDescription(UObject* Owne
 
 	if (BuildSettings->bGenerateLightmapUVs && VertexInstanceArray.Num() > 0)
 	{
-		if (BuildSettings->SrcLightmapIndex >= RenderMeshDescription->VertexInstanceAttributes().GetAttributeIndexCount<FVector2D>(MeshAttribute::VertexInstance::TextureCoordinate))
+		TVertexInstanceAttributeIndicesArray<FVector2D>& VertexInstanceUVs = RenderMeshDescription->VertexInstanceAttributes().GetAttributesSet<FVector2D>(MeshAttribute::VertexInstance::TextureCoordinate);
+		int32 NumIndices = VertexInstanceUVs.GetNumIndices();
+		//Verify the src light map channel
+		if (BuildSettings->SrcLightmapIndex >= NumIndices)
 		{
 			BuildSettings->SrcLightmapIndex = 0;
 		}
+		//Verify the destination light map channel
+		if (BuildSettings->DstLightmapIndex >= NumIndices)
+		{
+			//Make sure we do not add illegal UV Channel index
+			if (BuildSettings->DstLightmapIndex >= MAX_MESH_TEXTURE_COORDS_MD)
+			{
+				BuildSettings->DstLightmapIndex = MAX_MESH_TEXTURE_COORDS_MD - 1;
+			}
+
+			//Add some unused UVChannel to the mesh description for the lightmapUVs
+			VertexInstanceUVs.SetNumIndices(BuildSettings->DstLightmapIndex + 1);
+			BuildSettings->DstLightmapIndex = NumIndices;
+		}
 
 		FLayoutUV Packer(RenderMeshDescription, BuildSettings->SrcLightmapIndex, BuildSettings->DstLightmapIndex, BuildSettings->MinLightmapResolution);
 		Packer.SetVersion((FMeshDescriptionHelper::ELightmapUVVersion)(StaticMesh->LightmapUVVersion));
@@ -166,7 +182,6 @@ bool FMeshDescriptionHelper::IsValidOriginalMeshDescription()
 
 void FMeshDescriptionHelper::FindOverlappingCorners(TMultiMap<int32, int32>& OverlappingCorners, const UMeshDescription* MeshDescription, float ComparisonThreshold)
 {
-	BuildStatisticManager::FBuildStatisticScope StatScope(TEXT("FindOverlappingCorners took"));
 	//Empty the old data
 	OverlappingCorners.Reset();
 	
@@ -223,7 +238,6 @@ void FMeshDescriptionHelper::CreatePolygonNTB(UMeshDescription* MeshDescription,
 	TPolygonAttributeArray<FVector>& PolygonTangents = MeshDescription->PolygonAttributes().GetAttributes<FVector>(MeshAttribute::Polygon::Tangent);
 	TPolygonAttributeArray<FVector>& PolygonBinormals = MeshDescription->PolygonAttributes().GetAttributes<FVector>(MeshAttribute::Polygon::Binormal);
 
-	BuildStatisticManager::FBuildStatisticScope StatScope(TEXT("CreatePolygonNTB took"));
 	TArray<FPolygonID> DegeneratePolygons;
 
 	FVertexInstanceArray& VertexInstanceArray = MeshDescription->VertexInstances();
@@ -347,45 +361,90 @@ struct FVertexInfo
 	TArray<FEdgeID> EdgeIDs;
 };
 
-void RecursiveFillPolygonGroup(
-	  UMeshDescription* MeshDescription
-	, TMap<FPolygonID, FVertexInfo>& VertexInfoMap
-	, TArray<FPolygonID>& CurrentGroup
-	, FVertexInfo& CurrentVertexInfo
-	, TArray<FPolygonID>& ConsumedPolygon)
+namespace MeshDescriptionPrototype
 {
-	CurrentGroup.AddUnique(CurrentVertexInfo.PolygonID);
-	ConsumedPolygon.AddUnique(CurrentVertexInfo.PolygonID);
-
-	const TEdgeAttributeArray<bool>& EdgeHardnesses = MeshDescription->EdgeAttributes().GetAttributes<bool>(MeshAttribute::Edge::IsHard);
-
-	for (const FEdgeID EdgeID : CurrentVertexInfo.EdgeIDs)
+	void GetVertexConnectedPolygons(const UMeshDescription* MeshDescription, const FVertexID VertexID, TArray<FPolygonID>& OutConnectedPolygonIDs)
 	{
-		if (EdgeHardnesses[EdgeID])
+		OutConnectedPolygonIDs.Reset();
+
+		const FVertexInstanceArray& VertexInstances = MeshDescription->VertexInstances();
+		for (const FVertexInstanceID VertexInstanceID : MeshDescription->GetVertex(VertexID).VertexInstanceIDs)
 		{
-			//End of the group
-			continue;
+			OutConnectedPolygonIDs.Append(VertexInstances[VertexInstanceID].ConnectedPolygons);
 		}
-		for (const FPolygonID PolygonID : MeshDescription->GetEdgeConnectedPolygons(EdgeID))
+	}
+
+	void GetConnectedSoftEdges(const UMeshDescription* MeshDescription, const FVertexID VertexID, TArray<FEdgeID>& OutConnectedSoftEdges)
+	{
+		OutConnectedSoftEdges.Reset();
+
+		const TEdgeAttributeArray<bool>& EdgeHardnesses = MeshDescription->EdgeAttributes().GetAttributes<bool>(MeshAttribute::Edge::IsHard);
+		for (const FEdgeID ConnectedEdgeID : MeshDescription->GetVertex(VertexID).ConnectedEdgeIDs)
 		{
-			if (PolygonID == CurrentVertexInfo.PolygonID)
+			if (!EdgeHardnesses[ConnectedEdgeID])
 			{
-				continue;
+				OutConnectedSoftEdges.Add(ConnectedEdgeID);
 			}
-			//Add this polygon to the group
-			FVertexInfo& OtherVertexInfo = VertexInfoMap.FindOrAdd(PolygonID);
-			//Do not repeat polygons
-			if (!ConsumedPolygon.Contains(OtherVertexInfo.PolygonID))
+		}
+	}
+
+	void GetPolygonsInSameSoftEdgedGroup(const UMeshDescription* MeshDescription, const FVertexID VertexID, const FPolygonID PolygonID, TArray<FPolygonID>& OutPolygonIDs)
+	{
+		// The aim here is to determine which polygons form part of the same soft edged group as the polygons attached to this vertex instance.
+		// They should all contribute to the final vertex instance normal.
+
+		OutPolygonIDs.Reset();
+
+		// Get all polygons connected to this vertex.
+		static TArray<FPolygonID> ConnectedPolygons;
+		GetVertexConnectedPolygons(MeshDescription, VertexID, ConnectedPolygons);
+
+		// Cache a list of all soft edges which share this vertex.
+		// We're only interested in finding adjacent polygons which are not the other side of a hard edge.
+		static TArray<FEdgeID> ConnectedSoftEdges;
+		GetConnectedSoftEdges(MeshDescription, VertexID, ConnectedSoftEdges);
+
+		// Iterate through adjacent polygons which contain the given vertex, but don't cross a hard edge.
+		// Maintain a list of polygon IDs to be examined. Adjacents are added to the list if suitable.
+		// Add the start poly here.
+		static TArray<FPolygonID> PolygonsToCheck;
+		PolygonsToCheck.Reset();
+		PolygonsToCheck.Add(PolygonID);
+
+		const FEdgeArray& Edges = MeshDescription->Edges();
+		int32 Index = 0;
+		while (Index < PolygonsToCheck.Num())
+		{
+			const FPolygonID PolygonToCheck = PolygonsToCheck[Index];
+			Index++;
+
+			if (ConnectedPolygons.Contains(PolygonToCheck))
 			{
-				RecursiveFillPolygonGroup(MeshDescription, VertexInfoMap, CurrentGroup, OtherVertexInfo, ConsumedPolygon);
+				OutPolygonIDs.Add(PolygonToCheck);
+
+				// Now look at its adjacent polygons. If they are joined by a soft edge which includes the vertex we're interested in, we want to consider them.
+				// We take a shortcut by doing this process in reverse: we already know all the soft edges we are interested in, so check if any of them
+				// have the current polygon as an adjacent.
+				for (const FEdgeID ConnectedSoftEdge : ConnectedSoftEdges)
+				{
+					const FMeshEdge& Edge = Edges[ConnectedSoftEdge];
+					if (Edge.ConnectedPolygons.Contains(PolygonToCheck))
+					{
+						for (const FPolygonID AdjacentPolygon : Edge.ConnectedPolygons)
+						{
+							// Only add new polygons which haven't yet been added to the list. This prevents circular runs of polygons triggering infinite loops.
+							PolygonsToCheck.AddUnique(AdjacentPolygon);
+						}
+					}
+				}
 			}
 		}
 	}
-}
+
+} //End of namespace MeshDescriptionPrototype
 
 void FMeshDescriptionHelper::CreateNormals(UMeshDescription* MeshDescription, FMeshDescriptionHelper::ETangentOptions TangentOptions, bool bComputeTangent)
 {
-	BuildStatisticManager::FBuildStatisticScope StatScope(TEXT("CreateNormals took"));
 	//For each vertex compute the normals for every connected edges that are smooth betwween hard edges
 	//         H   A    B
 	//          \  ||  /
@@ -464,8 +523,39 @@ void FMeshDescriptionHelper::CreateNormals(UMeshDescription* MeshDescription, FM
 			}
 
 			int32 CurrentGroupIndex = Groups.AddZeroed();
-			FVertexInfo& VertexInfo = VertexInfoMap.FindOrAdd(Kvp.Key);
-			RecursiveFillPolygonGroup(MeshDescription, VertexInfoMap, Groups[CurrentGroupIndex], VertexInfo, ConsumedPolygon);
+			TArray<FPolygonID>& CurrentGroup = Groups[CurrentGroupIndex];
+			TArray<FPolygonID> PolygonQueue;
+			PolygonQueue.Add(Kvp.Key); //Use a queue to avoid recursive function
+			while(PolygonQueue.Num() > 0)
+			{
+				FPolygonID CurrentPolygonID = PolygonQueue.Pop(true);
+				FVertexInfo& CurrentVertexInfo = VertexInfoMap.FindOrAdd(CurrentPolygonID);
+				CurrentGroup.AddUnique(CurrentVertexInfo.PolygonID);
+				ConsumedPolygon.AddUnique(CurrentVertexInfo.PolygonID);
+				const TEdgeAttributeArray<bool>& EdgeHardnesses = MeshDescription->EdgeAttributes().GetAttributes<bool>(MeshAttribute::Edge::IsHard);
+				for (const FEdgeID EdgeID : CurrentVertexInfo.EdgeIDs)
+				{
+					if (EdgeHardnesses[EdgeID])
+					{
+						//End of the group
+						continue;
+					}
+					for (const FPolygonID PolygonID : MeshDescription->GetEdgeConnectedPolygons(EdgeID))
+					{
+						if (PolygonID == CurrentVertexInfo.PolygonID)
+						{
+							continue;
+						}
+						//Add this polygon to the group
+						FVertexInfo& OtherVertexInfo = VertexInfoMap.FindOrAdd(PolygonID);
+						//Do not repeat polygons
+						if (!ConsumedPolygon.Contains(OtherVertexInfo.PolygonID))
+						{
+							PolygonQueue.Add(PolygonID);
+						}
+					}
+				}
+			}
 		}
 
 		//Smooth every connected group
@@ -576,7 +666,6 @@ void FMeshDescriptionHelper::CreateNormals(UMeshDescription* MeshDescription, FM
 
 void FMeshDescriptionHelper::CreateMikktTangents(UMeshDescription* MeshDescription, FMeshDescriptionHelper::ETangentOptions TangentOptions)
 {
-	BuildStatisticManager::FBuildStatisticScope StatScope(TEXT("CreateMikktTangents took"));
 	bool bIgnoreDegenerateTriangles = (TangentOptions & FMeshDescriptionHelper::ETangentOptions::IgnoreDegenerateTriangles) != 0;
 
 	// we can use mikktspace to calculate the tangents
@@ -648,4 +737,419 @@ namespace MeshDescriptionMikktSpaceInterface
 		UV[0] = TexCoord.X;
 		UV[1] = TexCoord.Y;
 	}
+}
+
+//////////////////////////////////////////////////////////////////////////
+//Converters
+
+void FMeshDescriptionHelper::ConvertHardEdgesToSmoothGroup(const UMeshDescription* SourceMeshDescription, struct FRawMesh &DestinationRawMesh)
+{
+	TMap<FPolygonID, uint32> PolygonSmoothGroup;
+	PolygonSmoothGroup.Reserve(SourceMeshDescription->Polygons().Num());
+	TArray<bool> ConsumedPolygons;
+	ConsumedPolygons.AddZeroed(SourceMeshDescription->Polygons().Num());
+
+	TMap < FPolygonID, uint32> PolygonAvoidances;
+
+	const TEdgeAttributeArray<bool>& EdgeHardnesses = SourceMeshDescription->EdgeAttributes().GetAttributes<bool>(MeshAttribute::Edge::IsHard);
+
+	for (const FPolygonID PolygonID : SourceMeshDescription->Polygons().GetElementIDs())
+	{
+		if (ConsumedPolygons[PolygonID.GetValue()])
+		{
+			continue;
+		}
+		TArray<FPolygonID> ConnectedPolygons;
+		TArray<FPolygonID> LastConnectedPolygons;
+		ConnectedPolygons.Add(PolygonID);
+		LastConnectedPolygons.Add(FPolygonID::Invalid);
+		while (ConnectedPolygons.Num() > 0)
+		{
+			check(LastConnectedPolygons.Num() == ConnectedPolygons.Num());
+			FPolygonID LastPolygonID = LastConnectedPolygons.Pop(true);
+			FPolygonID CurrentPolygonID = ConnectedPolygons.Pop(true);
+			if (ConsumedPolygons[CurrentPolygonID.GetValue()])
+			{
+				continue;
+			}
+			TArray<FPolygonID> SoftEdgeNeigbors;
+			uint32& SmoothGroup = PolygonSmoothGroup.FindOrAdd(CurrentPolygonID);
+			uint32 AvoidSmoothGroup = 0;
+			uint32 NeighborSmoothGroup = 0;
+			const uint32 LastSmoothGroupValue = (LastPolygonID == FPolygonID::Invalid) ? 0 : PolygonSmoothGroup[LastPolygonID];
+			TArray<FEdgeID> PolygonEdges;
+			SourceMeshDescription->GetPolygonEdges(CurrentPolygonID, PolygonEdges);
+			for (const FEdgeID& EdgeID : PolygonEdges)
+			{
+				bool bIsHardEdge = EdgeHardnesses[EdgeID];
+				const TArray<FPolygonID>& EdgeConnectedPolygons = SourceMeshDescription->GetEdgeConnectedPolygons(EdgeID);
+				for (const FPolygonID& EdgePolygonID : EdgeConnectedPolygons)
+				{
+					if (EdgePolygonID == CurrentPolygonID)
+					{
+						continue;
+					}
+					uint32 SmoothValue = 0;
+					if (PolygonSmoothGroup.Contains(EdgePolygonID))
+					{
+						SmoothValue = PolygonSmoothGroup[EdgePolygonID];
+					}
+
+					if (bIsHardEdge) //Hard Edge
+					{
+						AvoidSmoothGroup |= SmoothValue;
+					}
+					else
+					{
+						NeighborSmoothGroup |= SmoothValue;
+						//Put all none hard edge polygon in the next iteration
+						if (!ConsumedPolygons[EdgePolygonID.GetValue()])
+						{
+							ConnectedPolygons.Add(EdgePolygonID);
+							LastConnectedPolygons.Add(CurrentPolygonID);
+						}
+						else
+						{
+							SoftEdgeNeigbors.Add(EdgePolygonID);
+						}
+					}
+				}
+			}
+
+			if (AvoidSmoothGroup != 0)
+			{
+				PolygonAvoidances.FindOrAdd(CurrentPolygonID) = AvoidSmoothGroup;
+				//find neighbor avoidance
+				for (FPolygonID& NeighborID : SoftEdgeNeigbors)
+				{
+					if(!PolygonAvoidances.Contains(NeighborID))
+					{
+						continue;
+					}
+					AvoidSmoothGroup |= PolygonAvoidances[NeighborID];
+				}
+				uint32 NewSmoothGroup = 1;
+				while ((NewSmoothGroup & AvoidSmoothGroup) != 0 && NewSmoothGroup < MAX_uint32)
+				{
+					//Shift the smooth group
+					NewSmoothGroup = NewSmoothGroup << 1;
+				}
+				SmoothGroup = NewSmoothGroup;
+				//Aply to all neighboard
+				for(FPolygonID& NeighborID : SoftEdgeNeigbors)
+				{
+					PolygonSmoothGroup[NeighborID] |= NewSmoothGroup;
+				}
+			}
+			else if (NeighborSmoothGroup != 0)
+			{
+				SmoothGroup |= LastSmoothGroupValue | NeighborSmoothGroup;
+			}
+			else
+			{
+				SmoothGroup = 1;
+			}
+			ConsumedPolygons[CurrentPolygonID.GetValue()] = true;
+		}
+	}
+	//Now we have to put the data into the RawMesh
+	int32 TriangleIndex = 0;
+	for (const FPolygonID PolygonID : SourceMeshDescription->Polygons().GetElementIDs())
+	{
+		uint32 PolygonSmoothValue = PolygonSmoothGroup[PolygonID];
+		const TArray<FMeshTriangle>& Triangles = SourceMeshDescription->GetPolygonTriangles(PolygonID);
+		for (const FMeshTriangle& MeshTriangle : Triangles)
+		{
+			DestinationRawMesh.FaceSmoothingMasks[TriangleIndex++] = PolygonSmoothValue;
+		}
+	}
+}
+
+
+
+
+void FMeshDescriptionHelper::ConvertSmoothGroupToHardEdges(const struct FRawMesh &SourceRawMesh, UMeshDescription* DestinationMeshDescription)
+{
+	TEdgeAttributeArray<bool>& EdgeHardnesses = DestinationMeshDescription->EdgeAttributes().GetAttributes<bool>(MeshAttribute::Edge::IsHard);
+
+	TArray<bool> ConsumedPolygons;
+	ConsumedPolygons.AddZeroed(DestinationMeshDescription->Polygons().Num());
+	for (const FPolygonID PolygonID : DestinationMeshDescription->Polygons().GetElementIDs())
+	{
+		if (ConsumedPolygons[PolygonID.GetValue()])
+		{
+			continue;
+		}
+		TArray<FPolygonID> ConnectedPolygons;
+		ConnectedPolygons.Add(PolygonID);
+		while (ConnectedPolygons.Num() > 0)
+		{
+			FPolygonID CurrentPolygonID = ConnectedPolygons.Pop(true);
+			int32 CurrentPolygonIDValue = CurrentPolygonID.GetValue();
+			check(SourceRawMesh.FaceSmoothingMasks.IsValidIndex(CurrentPolygonIDValue));
+			const uint32 ReferenceSmoothGroup = SourceRawMesh.FaceSmoothingMasks[CurrentPolygonIDValue];
+			TArray<FEdgeID> PolygonEdges;
+			DestinationMeshDescription->GetPolygonEdges(CurrentPolygonID, PolygonEdges);
+			for (const FEdgeID& EdgeID : PolygonEdges)
+			{
+				const bool bIsHardEdge = EdgeHardnesses[EdgeID];
+				if (bIsHardEdge)
+				{
+					continue;
+				}
+				const TArray<FPolygonID>& EdgeConnectedPolygons = DestinationMeshDescription->GetEdgeConnectedPolygons(EdgeID);
+				for (const FPolygonID& EdgePolygonID : EdgeConnectedPolygons)
+				{
+					int32 EdgePolygonIDValue = EdgePolygonID.GetValue();
+					if (EdgePolygonID == CurrentPolygonID || ConsumedPolygons[EdgePolygonIDValue])
+					{
+						continue;
+					}
+					check(SourceRawMesh.FaceSmoothingMasks.IsValidIndex(EdgePolygonIDValue));
+					const uint32 TestSmoothGroup = SourceRawMesh.FaceSmoothingMasks[EdgePolygonIDValue];
+					if ((TestSmoothGroup & ReferenceSmoothGroup) == 0)
+					{
+						EdgeHardnesses[EdgeID] = true;
+						break;
+					}
+					else
+					{
+						ConnectedPolygons.Add(EdgePolygonID);
+					}
+				}
+			}
+			ConsumedPolygons[CurrentPolygonID.GetValue()] = true;
+		}
+	}
+}
+
+void FMeshDescriptionHelper::ConverToRawMesh(const UMeshDescription* SourceMeshDescription, struct FRawMesh &DestinationRawMesh)
+{
+	DestinationRawMesh.Empty();
+
+	//Gather all array data
+	const TVertexAttributeArray<FVector>& VertexPositions = SourceMeshDescription->VertexAttributes().GetAttributes<FVector>(MeshAttribute::Vertex::Position);
+	//const TVertexAttributeArray<float>& VertexCornerSharpness = SourceMeshDescription->VertexAttributes().GetAttributes<float>(MeshAttribute::Vertex::CornerSharpness);
+
+	const TVertexInstanceAttributeArray<FVector>& VertexInstanceNormals = SourceMeshDescription->VertexInstanceAttributes().GetAttributes<FVector>(MeshAttribute::VertexInstance::Normal);
+	const TVertexInstanceAttributeArray<FVector>& VertexInstanceTangents = SourceMeshDescription->VertexInstanceAttributes().GetAttributes<FVector>(MeshAttribute::VertexInstance::Tangent);
+	const TVertexInstanceAttributeArray<float>& VertexInstanceBinormalSigns = SourceMeshDescription->VertexInstanceAttributes().GetAttributes<float>(MeshAttribute::VertexInstance::BinormalSign);
+	const TVertexInstanceAttributeArray<FVector4>& VertexInstanceColors = SourceMeshDescription->VertexInstanceAttributes().GetAttributes<FVector4>(MeshAttribute::VertexInstance::Color);
+	const TVertexInstanceAttributeIndicesArray<FVector2D>& VertexInstanceUVs = SourceMeshDescription->VertexInstanceAttributes().GetAttributesSet<FVector2D>(MeshAttribute::VertexInstance::TextureCoordinate);
+
+	//const TEdgeAttributeArray<bool>& EdgeHardnesses = SourceMeshDescription->EdgeAttributes().GetAttributes<bool>(MeshAttribute::Edge::IsHard);
+	//const TEdgeAttributeArray<float>& EdgeCreaseSharpnesses = SourceMeshDescription->EdgeAttributes().GetAttributes<float>(MeshAttribute::Edge::CreaseSharpness);
+
+	//const TPolygonAttributeArray<FVector>& PolygonNormals = SourceMeshDescription->PolygonAttributes().GetAttributes<FVector>(MeshAttribute::Polygon::Normal);
+	//const TPolygonAttributeArray<FVector>& PolygonTangents = SourceMeshDescription->PolygonAttributes().GetAttributes<FVector>(MeshAttribute::Polygon::Tangent);
+	//const TPolygonAttributeArray<FVector>& PolygonBiNormals = SourceMeshDescription->PolygonAttributes().GetAttributes<FVector>(MeshAttribute::Polygon::Binormal);
+	//const TPolygonAttributeArray<FVector>& PolygonCenters = SourceMeshDescription->PolygonAttributes().GetAttributes<FVector>(MeshAttribute::Polygon::Center);
+
+	//const TPolygonGroupAttributeArray<FSoftObjectPath>& PolygonGroupMaterialAssets = SourceMeshDescription->PolygonGroupAttributes().GetAttributes<FSoftObjectPath>(MeshAttribute::PolygonGroup::MaterialAsset);
+	//const TPolygonGroupAttributeArray<FName>& PolygonGroupMaterialSlotNames = SourceMeshDescription->PolygonGroupAttributes().GetAttributes<FName>(MeshAttribute::PolygonGroup::MaterialSlotName);
+	//const TPolygonGroupAttributeArray<FName>& PolygonGroupImportedMaterialSlotNames = SourceMeshDescription->PolygonGroupAttributes().GetAttributes<FName>(MeshAttribute::PolygonGroup::ImportedMaterialSlotName);
+	//const TPolygonGroupAttributeArray<bool>& PolygonGroupCollision = SourceMeshDescription->PolygonGroupAttributes().GetAttributes<bool>(MeshAttribute::PolygonGroup::EnableCollision);
+	//const TPolygonGroupAttributeArray<bool>& PolygonGroupCastShadow = SourceMeshDescription->PolygonGroupAttributes().GetAttributes<bool>(MeshAttribute::PolygonGroup::CastShadow);
+
+	DestinationRawMesh.VertexPositions.AddZeroed(SourceMeshDescription->Vertices().Num());
+	for (const FVertexID& VertexID : SourceMeshDescription->Vertices().GetElementIDs())
+	{
+		int32 VertexIDValue = VertexID.GetValue();
+		DestinationRawMesh.VertexPositions[VertexIDValue] = VertexPositions[VertexID];
+	}
+	int32 VertexInstanceNumber = SourceMeshDescription->VertexInstances().Num();
+	DestinationRawMesh.WedgeColors.AddZeroed(VertexInstanceNumber);
+	DestinationRawMesh.WedgeIndices.AddZeroed(VertexInstanceNumber);
+	DestinationRawMesh.WedgeTangentX.AddZeroed(VertexInstanceNumber);
+	DestinationRawMesh.WedgeTangentY.AddZeroed(VertexInstanceNumber);
+	DestinationRawMesh.WedgeTangentZ.AddZeroed(VertexInstanceNumber);
+	int32 ExistingUVCount = VertexInstanceUVs.GetNumIndices();
+	for (int32 UVIndex = 0; UVIndex < ExistingUVCount; ++UVIndex)
+	{
+		DestinationRawMesh.WedgeTexCoords[UVIndex].AddZeroed(VertexInstanceNumber);
+	}
+
+	int32 TriangleNumber = 0;
+	for (const FPolygonID& PolygonID : SourceMeshDescription->Polygons().GetElementIDs())
+	{
+		TriangleNumber += SourceMeshDescription->GetPolygonTriangles(PolygonID).Num();
+	}
+	DestinationRawMesh.FaceMaterialIndices.AddZeroed(TriangleNumber);
+	DestinationRawMesh.FaceSmoothingMasks.AddZeroed(TriangleNumber);
+
+	for (const FPolygonID& PolygonID : SourceMeshDescription->Polygons().GetElementIDs())
+	{
+		const FPolygonGroupID& PolygonGroupID = SourceMeshDescription->GetPolygonPolygonGroup(PolygonID);
+		int32 PolygonIDValue = PolygonID.GetValue();
+		const TArray<FMeshTriangle>& Triangles = SourceMeshDescription->GetPolygonTriangles(PolygonID);
+		for (const FMeshTriangle& MeshTriangle : Triangles)
+		{
+			DestinationRawMesh.FaceMaterialIndices[PolygonIDValue] = PolygonGroupID.GetValue();
+			DestinationRawMesh.FaceSmoothingMasks[PolygonIDValue] = 0; //Conversion of soft/hard to smooth mask is done after the geometry is converted
+			for (int32 Corner = 0; Corner < 3; ++Corner)
+			{
+				const FVertexInstanceID VertexInstanceID = MeshTriangle.GetVertexInstanceID(Corner);
+				const int32 VertexInstanceIDValue = VertexInstanceID.GetValue();
+				DestinationRawMesh.WedgeColors[VertexInstanceIDValue] = FLinearColor(VertexInstanceColors[VertexInstanceID]).ToFColor(true);
+				DestinationRawMesh.WedgeIndices[VertexInstanceIDValue] = SourceMeshDescription->GetVertexInstanceVertex(VertexInstanceID).GetValue();
+				DestinationRawMesh.WedgeTangentX[VertexInstanceIDValue] = VertexInstanceTangents[VertexInstanceID];
+				DestinationRawMesh.WedgeTangentY[VertexInstanceIDValue] = FVector::CrossProduct(VertexInstanceNormals[VertexInstanceID], VertexInstanceTangents[VertexInstanceID]).GetSafeNormal() * VertexInstanceBinormalSigns[VertexInstanceID];
+				DestinationRawMesh.WedgeTangentZ[VertexInstanceIDValue] = VertexInstanceNormals[VertexInstanceID];
+				for (int32 UVIndex = 0; UVIndex < ExistingUVCount; ++UVIndex)
+				{
+					DestinationRawMesh.WedgeTexCoords[UVIndex][VertexInstanceIDValue] = VertexInstanceUVs.GetArrayForIndex(UVIndex)[VertexInstanceID];
+				}
+			}
+		}
+	}
+	//Convert the smoothgroup
+	ConvertHardEdgesToSmoothGroup(SourceMeshDescription, DestinationRawMesh);
+}
+
+void FMeshDescriptionHelper::ConverFromRawMesh(const struct FRawMesh &SourceRawMesh, UMeshDescription* DestinationMeshDescription)
+{
+	DestinationMeshDescription->Empty();
+	//Gather all array data
+	TVertexAttributeArray<FVector>& VertexPositions = DestinationMeshDescription->VertexAttributes().GetAttributes<FVector>(MeshAttribute::Vertex::Position);
+	//TVertexAttributeArray<float>& VertexCornerSharpness = DestinationMeshDescription->VertexAttributes().GetAttributes<float>(MeshAttribute::Vertex::CornerSharpness);
+
+	TVertexInstanceAttributeArray<FVector>& VertexInstanceNormals = DestinationMeshDescription->VertexInstanceAttributes().GetAttributes<FVector>(MeshAttribute::VertexInstance::Normal);
+	TVertexInstanceAttributeArray<FVector>& VertexInstanceTangents = DestinationMeshDescription->VertexInstanceAttributes().GetAttributes<FVector>(MeshAttribute::VertexInstance::Tangent);
+	TVertexInstanceAttributeArray<float>& VertexInstanceBinormalSigns = DestinationMeshDescription->VertexInstanceAttributes().GetAttributes<float>(MeshAttribute::VertexInstance::BinormalSign);
+	TVertexInstanceAttributeArray<FVector4>& VertexInstanceColors = DestinationMeshDescription->VertexInstanceAttributes().GetAttributes<FVector4>(MeshAttribute::VertexInstance::Color);
+	TVertexInstanceAttributeIndicesArray<FVector2D>& VertexInstanceUVs = DestinationMeshDescription->VertexInstanceAttributes().GetAttributesSet<FVector2D>(MeshAttribute::VertexInstance::TextureCoordinate);
+
+	//TEdgeAttributeArray<bool>& EdgeHardnesses = DestinationMeshDescription->EdgeAttributes().GetAttributes<bool>(MeshAttribute::Edge::IsHard);
+	//TEdgeAttributeArray<float>& EdgeCreaseSharpnesses = DestinationMeshDescription->EdgeAttributes().GetAttributes<float>(MeshAttribute::Edge::CreaseSharpness);
+
+	//TPolygonAttributeArray<FVector>& PolygonNormals = DestinationMeshDescription->PolygonAttributes().GetAttributes<FVector>(MeshAttribute::Polygon::Normal);
+	//TPolygonAttributeArray<FVector>& PolygonTangents = DestinationMeshDescription->PolygonAttributes().GetAttributes<FVector>(MeshAttribute::Polygon::Tangent);
+	//TPolygonAttributeArray<FVector>& PolygonBiNormals = DestinationMeshDescription->PolygonAttributes().GetAttributes<FVector>(MeshAttribute::Polygon::Binormal);
+	//TPolygonAttributeArray<FVector>& PolygonCenters = DestinationMeshDescription->PolygonAttributes().GetAttributes<FVector>(MeshAttribute::Polygon::Center);
+
+	TPolygonGroupAttributeArray<FSoftObjectPath>& PolygonGroupMaterialAssets = DestinationMeshDescription->PolygonGroupAttributes().GetAttributes<FSoftObjectPath>(MeshAttribute::PolygonGroup::MaterialAsset);
+	TPolygonGroupAttributeArray<FName>& PolygonGroupMaterialSlotNames = DestinationMeshDescription->PolygonGroupAttributes().GetAttributes<FName>(MeshAttribute::PolygonGroup::MaterialSlotName);
+	TPolygonGroupAttributeArray<FName>& PolygonGroupImportedMaterialSlotNames = DestinationMeshDescription->PolygonGroupAttributes().GetAttributes<FName>(MeshAttribute::PolygonGroup::ImportedMaterialSlotName);
+	TPolygonGroupAttributeArray<bool>& PolygonGroupCollision = DestinationMeshDescription->PolygonGroupAttributes().GetAttributes<bool>(MeshAttribute::PolygonGroup::EnableCollision);
+	TPolygonGroupAttributeArray<bool>& PolygonGroupCastShadow = DestinationMeshDescription->PolygonGroupAttributes().GetAttributes<bool>(MeshAttribute::PolygonGroup::CastShadow);
+
+	int32 NumTexCoords = 0;
+	int32 MaxTexCoords = FMath::Min<int32>(MAX_MESH_TEXTURE_COORDS, MAX_STATIC_TEXCOORDS);
+	TArray<int32> TextureCoordinnateRemapIndex;
+	TextureCoordinnateRemapIndex.AddZeroed(MaxTexCoords);
+	for (int32 TextureCoordinnateIndex = 0; TextureCoordinnateIndex < MaxTexCoords; ++TextureCoordinnateIndex)
+	{
+		TextureCoordinnateRemapIndex[TextureCoordinnateIndex] = INDEX_NONE;
+		if (SourceRawMesh.WedgeTexCoords[TextureCoordinnateIndex].Num() == SourceRawMesh.WedgeIndices.Num())
+		{
+			TextureCoordinnateRemapIndex[TextureCoordinnateIndex] = NumTexCoords;
+			NumTexCoords++;
+		}
+	}
+	VertexInstanceUVs.SetNumIndices(NumTexCoords);
+	for (int32 VertexIndex = 0; VertexIndex < SourceRawMesh.VertexPositions.Num(); ++VertexIndex)
+	{
+		FVertexID VertexID = DestinationMeshDescription->CreateVertex();
+		VertexPositions[VertexID] = SourceRawMesh.VertexPositions[VertexIndex];
+	}
+
+	bool bHasColors = SourceRawMesh.WedgeColors.Num() > 0;
+	bool bHasTangents = SourceRawMesh.WedgeTangentX.Num() > 0 && SourceRawMesh.WedgeTangentY.Num() > 0;
+	bool bHasNormals = SourceRawMesh.WedgeTangentZ.Num() > 0;
+
+	TArray<FPolygonGroupID> PolygonGroups;
+
+	//Triangles
+	int32 TriangleCount = SourceRawMesh.WedgeIndices.Num() / 3;
+	for (int32 TriangleIndex = 0; TriangleIndex < TriangleCount; ++TriangleIndex)
+	{
+		int32 VerticeIndexBase = TriangleIndex * 3;
+
+		//PolygonGroup
+		FPolygonGroupID PolygonGroupID(SourceRawMesh.FaceMaterialIndices[TriangleIndex]);
+		if (PolygonGroupID.GetValue() >= PolygonGroups.Num())
+		{
+			for (int32 AddMaterialIndex = PolygonGroups.Num(); AddMaterialIndex <= PolygonGroupID.GetValue(); ++AddMaterialIndex)
+			{
+				FPolygonGroupID AddedPolygonGroupID = DestinationMeshDescription->CreatePolygonGroup();
+				check(AddMaterialIndex == AddedPolygonGroupID.GetValue());
+				PolygonGroupCastShadow[AddedPolygonGroupID] = true;
+				PolygonGroupCollision[AddedPolygonGroupID] = true;
+				PolygonGroupImportedMaterialSlotNames[AddedPolygonGroupID] = NAME_None;
+				PolygonGroupMaterialSlotNames[AddedPolygonGroupID] = NAME_None;
+				PolygonGroupMaterialAssets[AddedPolygonGroupID] = FSoftObjectPath(UMaterial::GetDefaultMaterial(MD_Surface));
+				PolygonGroups.Add(AddedPolygonGroupID);
+			}
+		}
+
+		for (int32 Corner = 0; Corner < 3; ++Corner)
+		{
+			int32 VerticeIndex = VerticeIndexBase + Corner;
+			FVertexID VertexID(SourceRawMesh.WedgeIndices[VerticeIndex]);
+			FVertexInstanceID VertexInstanceID = DestinationMeshDescription->CreateVertexInstance(VertexID);
+			VertexInstanceColors[VertexInstanceID] = bHasColors ? FLinearColor::FromSRGBColor(SourceRawMesh.WedgeColors[VerticeIndex]) : FLinearColor::White;
+			VertexInstanceTangents[VertexInstanceID] = bHasTangents ? SourceRawMesh.WedgeTangentX[VerticeIndex] : FVector(ForceInitToZero);
+			VertexInstanceBinormalSigns[VertexInstanceID] = bHasTangents ? GetBasisDeterminantSign(SourceRawMesh.WedgeTangentX[VerticeIndex].GetSafeNormal(), SourceRawMesh.WedgeTangentY[VerticeIndex].GetSafeNormal(), SourceRawMesh.WedgeTangentZ[VerticeIndex].GetSafeNormal()) : 0.0f;
+			VertexInstanceNormals[VertexInstanceID] = bHasNormals ? SourceRawMesh.WedgeTangentZ[VerticeIndex] : FVector(ForceInitToZero);
+			for (int32 TextureCoordinnateIndex = 0; TextureCoordinnateIndex < NumTexCoords; ++TextureCoordinnateIndex)
+			{
+				int32 TextureCoordIndex = TextureCoordinnateRemapIndex[TextureCoordinnateIndex];
+				if (TextureCoordIndex == INDEX_NONE)
+				{
+					continue;
+				}
+				TMeshAttributeArray<FVector2D, FVertexInstanceID>& Uvs = VertexInstanceUVs.GetArrayForIndex(TextureCoordIndex);
+				Uvs[VertexInstanceID] = SourceRawMesh.WedgeTexCoords[TextureCoordinnateIndex][VerticeIndex];
+			}
+		}
+		
+		//Create the polygon edges
+		TArray<UMeshDescription::FContourPoint> Contours;
+		// Add the edges of this triangle
+		for (uint32 Corner = 0; Corner < 3; ++Corner)
+		{
+			int32 ContourPointIndex = Contours.AddDefaulted();
+			UMeshDescription::FContourPoint& ContourPoint = Contours[ContourPointIndex];
+			//Find the matching edge ID
+			int32 CornerIndices[2];
+			CornerIndices[0] = (Corner + 0) % 3;
+			CornerIndices[1] = (Corner + 1) % 3;
+
+			FVertexID EdgeVertexIDs[2];
+			EdgeVertexIDs[0] = DestinationMeshDescription->GetVertexInstanceVertex(FVertexInstanceID(VerticeIndexBase + CornerIndices[0]));
+			EdgeVertexIDs[1] = DestinationMeshDescription->GetVertexInstanceVertex(FVertexInstanceID(VerticeIndexBase + CornerIndices[1]));
+
+			FEdgeID MatchEdgeId = DestinationMeshDescription->GetVertexPairEdge(EdgeVertexIDs[0], EdgeVertexIDs[1]);
+			if (MatchEdgeId == FEdgeID::Invalid)
+			{
+				MatchEdgeId = DestinationMeshDescription->CreateEdge(EdgeVertexIDs[0], EdgeVertexIDs[1]);
+			}
+			ContourPoint.EdgeID = MatchEdgeId;
+			ContourPoint.VertexInstanceID = FVertexInstanceID(VerticeIndexBase + CornerIndices[0]);
+
+			//TODO Edges smoothing
+		}
+
+		const FPolygonID NewPolygonID = DestinationMeshDescription->CreatePolygon(PolygonGroupID, Contours);
+		int32 NewTriangleIndex = DestinationMeshDescription->GetPolygonTriangles(NewPolygonID).AddDefaulted();
+		FMeshTriangle& NewTriangle = DestinationMeshDescription->GetPolygonTriangles(NewPolygonID)[NewTriangleIndex];
+		for (int32 Corner = 0; Corner < 3; ++Corner)
+		{
+			FVertexInstanceID VertexInstanceID(VerticeIndexBase + Corner);
+			NewTriangle.SetVertexInstanceID(Corner, VertexInstanceID);
+		}
+	}
+	CreatePolygonNTB(DestinationMeshDescription, 0.0f);
+
+	if (!bHasNormals || !bHasTangents)
+	{
+		//Create the missing normals and tangents
+		if (!bHasNormals)
+		{
+			CreateNormals(DestinationMeshDescription, ETangentOptions::BlendOverlappingNormals, false);
+		}
+		CreateMikktTangents(DestinationMeshDescription, ETangentOptions::BlendOverlappingNormals);
+	}
+
+	ConvertSmoothGroupToHardEdges(SourceRawMesh, DestinationMeshDescription);
 }
\ No newline at end of file
diff --git a/Engine/Source/Runtime/MeshBuilder/Private/MeshDescriptionHelper.h b/Engine/Source/Runtime/MeshBuilder/Private/MeshDescriptionHelper.h
--- a/Engine/Source/Runtime/MeshBuilder/Private/MeshDescriptionHelper.h
+++ b/Engine/Source/Runtime/MeshBuilder/Private/MeshDescriptionHelper.h
@@ -14,7 +14,8 @@ struct FVertexInstanceID;
 
 enum
 {
-	MAX_MESH_TEXTURE_COORDS = 8,
+	//Remove the _MD when FRawMesh will be remove
+	MAX_MESH_TEXTURE_COORDS_MD = 8,
 };
 
 class FMeshDescriptionHelper
@@ -51,14 +52,22 @@ public:
 
 	const TMultiMap<int32, int32>& GetOverlappingCorners() const { return OverlappingCorners; }
 
+	/** Convert this mesh description into the old FRawMesh format*/
+	static void ConverToRawMesh(const UMeshDescription* SourceMeshDescription, struct FRawMesh &DestinationRawMesh);
+	/** Convert old FRawMesh format to MeshDescription*/
+	static void ConverFromRawMesh(const struct FRawMesh &SourceRawMesh, UMeshDescription* DestinationMeshDescription);
+
 private:
 
 	//////////////////////////////////////////////////////////////////////////
 	//PRIVATE function declarations
 
-	void CreateNormals(UMeshDescription* MeshDescription, ETangentOptions TangentOptions, bool bComputeTangent);
-	void CreateMikktTangents(UMeshDescription* MeshDescription, ETangentOptions TangentOptions);
-	void CreatePolygonNTB(UMeshDescription* MeshDescription, float ComparisonThreshold);
+	static void ConvertHardEdgesToSmoothGroup(const UMeshDescription* SourceMeshDescription, struct FRawMesh &DestinationRawMesh);
+	static void ConvertSmoothGroupToHardEdges(const struct FRawMesh &SourceRawMesh, UMeshDescription* DestinationMeshDescription);
+
+	static void CreateNormals(UMeshDescription* MeshDescription, ETangentOptions TangentOptions, bool bComputeTangent);
+	static void CreateMikktTangents(UMeshDescription* MeshDescription, ETangentOptions TangentOptions);
+	static void CreatePolygonNTB(UMeshDescription* MeshDescription, float ComparisonThreshold);
 
 	//////////////////////////////////////////////////////////////////////////
 	//PRIVATE class members
diff --git a/Engine/Source/Runtime/MeshBuilder/Private/StaticMeshBuilder.cpp b/Engine/Source/Runtime/MeshBuilder/Private/StaticMeshBuilder.cpp
--- a/Engine/Source/Runtime/MeshBuilder/Private/StaticMeshBuilder.cpp
+++ b/Engine/Source/Runtime/MeshBuilder/Private/StaticMeshBuilder.cpp
@@ -11,9 +11,8 @@
 #include "Components.h"
 #include "IMeshReductionManagerModule.h"
 #include "MeshBuild.h"
-#include "BuildStatisticManager.h"
 
-DEFINE_LOG_CATEGORY(LogBuildStatistic);
+DEFINE_LOG_CATEGORY(LogStaticMeshBuilder);
 
 //////////////////////////////////////////////////////////////////////////
 //Local functions definition
@@ -43,10 +42,10 @@ FStaticMeshBuilder::FStaticMeshBuilder()
 
 bool FStaticMeshBuilder::Build(UStaticMesh* StaticMesh, const FStaticMeshLODGroup& LODGroup)
 {
-	BuildStatisticManager::FBuildStatisticScope StatScope(TEXT("Building StaticMesh took"));
 	if (StaticMesh->GetOriginalMeshDescription(0) == nullptr)
 	{
-		//TODO: Warn the user that there is no mesh description data
+		//Warn the user that there is no mesh description data
+		UE_LOG(LogStaticMeshBuilder, Error, TEXT("Cannot find a valid mesh description to build the asset."));
 		return false;
 	}
 	StaticMesh->RenderData->AllocateLODResources(StaticMesh->SourceModels.Num());
@@ -57,31 +56,54 @@ bool FStaticMeshBuilder::Build(UStaticMesh* StaticMesh, const FStaticMeshLODGrou
 	for (int32 LodIndex = 0; LodIndex < StaticMesh->SourceModels.Num(); ++LodIndex)
 	{
 		FMeshBuildSettings& LODBuildSettings = StaticMesh->SourceModels[LodIndex].BuildSettings;
-		const UMeshDescription* OriginalMeshDescription = StaticMesh->GetOriginalMeshDescription(LodIndex);
+		UMeshDescription* OriginalMeshDescription = StaticMesh->GetOriginalMeshDescription(LodIndex);
 		FMeshDescriptionHelper MeshDescriptionHelper(&LODBuildSettings, OriginalMeshDescription);
-		UMeshDescription* MeshDescription = MeshDescriptionHelper.GetRenderMeshDescription(StaticMesh);
+		
+		const FStaticMeshSourceModel& SrcModel = StaticMesh->SourceModels[LodIndex];
+		FMeshReductionSettings ReductionSettings = LODGroup.GetSettings(SrcModel.ReductionSettings, LodIndex);
+
+		UMeshDescription* MeshDescription = nullptr;
+		bool bUseReduction = (ReductionSettings.PercentTriangles < 1.0f || ReductionSettings.MaxDeviation > 0.0f) && (StaticMesh->GetMeshDescription(0) != nullptr);
+
+		if (OriginalMeshDescription != nullptr)
+		{
+			OriginalMeshDescription->TriangulateMesh(); //Make sure the original mesh description is triangulate.
+			MeshDescription = MeshDescriptionHelper.GetRenderMeshDescription(StaticMesh);
+		}
+		
+		if(MeshDescription == nullptr)
+		{
+			if (bUseReduction)
+			{
+				//Create an empty mesh description that the reduce will fill
+				MeshDescription = NewObject<UMeshDescription>(StaticMesh, NAME_None);
+			}
+			else if(StaticMesh->GetMeshDescription(0) != nullptr)
+			{
+				//Duplicate the lodindex 0 we have a 100% reduction which is like a duplicate
+				MeshDescription = Cast<UMeshDescription>(StaticDuplicateObject(StaticMesh->GetMeshDescription(0), StaticMesh, NAME_None, RF_NoFlags));
+			}
+		}
 
 		check(MeshDescription != nullptr);
 		StaticMesh->SetMeshDescription(LodIndex, MeshDescription);
 		
 		//Reduce LODs
-		const FStaticMeshSourceModel& SrcModel = StaticMesh->SourceModels[LodIndex];
-		FMeshReductionSettings ReductionSettings = LODGroup.GetSettings(SrcModel.ReductionSettings, LodIndex);
-		if (ReductionSettings.PercentTriangles < 1.0f || ReductionSettings.MaxDeviation > 0.0f)
+		if (bUseReduction)
 		{
 			float OverlappingThreshold = LODBuildSettings.bRemoveDegenerates ? THRESH_POINTS_ARE_SAME : 0.0f;
 			TMultiMap<int32, int32> OverlappingCorners;
-			MeshDescriptionHelper.FindOverlappingCorners(OverlappingCorners, StaticMesh->GetMeshDescription(LodIndex), OverlappingThreshold);
+			MeshDescriptionHelper.FindOverlappingCorners(OverlappingCorners, StaticMesh->GetMeshDescription(0), OverlappingThreshold);
 			MeshDescriptionHelper.ReduceLOD(StaticMesh->GetMeshDescription(0), StaticMesh->GetMeshDescription(LodIndex), ReductionSettings, OverlappingCorners);
 			// Recompute adjacency information. Since we change the vertices when we reduce
-			MeshDescriptionHelper.FindOverlappingCorners(StaticMesh->GetMeshDescription(0), OverlappingThreshold);
+			MeshDescriptionHelper.FindOverlappingCorners(StaticMesh->GetMeshDescription(LodIndex), OverlappingThreshold);
 		}
 
 		const FPolygonGroupArray& PolygonGroups = MeshDescription->PolygonGroups();
 
 		FStaticMeshLODResources& StaticMeshLOD = StaticMeshRenderData.LODResources[LodIndex];
 
-		//TODO: discover degenerate triangle with this threshold
+		//discover degenerate triangle with this threshold
 		float VertexComparisonThreshold = LODBuildSettings.bRemoveDegenerates ? THRESH_POINTS_ARE_SAME : 0.0f;
 
 		//Build new vertex buffers
@@ -343,13 +365,13 @@ void BuildVertexBuffer(
 	, float VertexComparisonThreshold
 	, TArray<int32>& RemapVerts)
 {
-	BuildStatisticManager::FBuildStatisticScope StatScope(TEXT("BuildVertexBuffer took"));
 	FVertexArray& Vertices = MeshDescription->Vertices();
 	FVertexInstanceArray& VertexInstances = MeshDescription->VertexInstances();
 	FPolygonGroupArray& PolygonGroupArray = MeshDescription->PolygonGroups();
 	FPolygonArray& PolygonArray = MeshDescription->Polygons();
 	
 	OutWedgeMap.Reset();
+	OutWedgeMap.AddZeroed(VertexInstances.Num());
 
 	TArray<int32> RemapVertexInstanceID;
 	// set up vertex buffer elements
@@ -389,24 +411,30 @@ void BuildVertexBuffer(
 		FStaticMeshSection& StaticMeshSection = StaticMeshLOD.Sections.Last();
 
 		StaticMeshSection.FirstIndex = IndexBuffer.Num();
-		// @todo: is currently assuming all polygons are triangles.
-		// At very least, this should be summing the sizes of all PolygonTriangles arrays.
-		// MeshDescription->GetPolygonTriangles() should always return a valid array of at least one triangle if the mesh description is valid.
-		StaticMeshSection.NumTriangles = Polygons.Num();
-
-		const int32 MaterialIndex = StaticMesh->GetMaterialIndex(PolygonGroupImportedMaterialSlotNames[PolygonGroupID]);
-		check(MaterialIndex != INDEX_NONE);
-		check(StaticMesh->StaticMaterials[MaterialIndex].MaterialInterface == Cast<UMaterialInterface>(PolygonGroupMaterialAssets[PolygonGroupID].TryLoad()));
-		StaticMeshSection.MaterialIndex = MaterialIndex;
-		StaticMeshSection.bEnableCollision = PolygonGroupCollision[PolygonGroupID];
-		StaticMeshSection.bCastShadow = PolygonGroupCastShadow[PolygonGroupID];
+		StaticMeshSection.NumTriangles = 0;
+		for (const FPolygonID& PolygonID : Polygons)
+		{
+			StaticMeshSection.NumTriangles += (uint32)MeshDescription->GetPolygonTriangles(PolygonID).Num();
+		}
+
+		StaticMeshSection.MaterialIndex = StaticMesh->GetMaterialIndexFromImportedMaterialSlotName(PolygonGroupImportedMaterialSlotNames[PolygonGroupID]);
+		
+		if (!StaticMesh->StaticMaterials.IsValidIndex(StaticMeshSection.MaterialIndex))
+		{
+			if (StaticMesh->StaticMaterials.IsValidIndex(PolygonGroupID.GetValue()))
+			{
+				StaticMeshSection.MaterialIndex = PolygonGroupID.GetValue();
+			}
+			else
+			{
+				StaticMeshSection.MaterialIndex = 0;
+			}
+		}
 
 		if (LodIndex > 0)
 		{
 			//Set the overwrite section info map
 			FMeshSectionInfo SectionInfo = StaticMesh->SectionInfoMap.Get(LodIndex, SectionIndex);
-			SectionInfo.bCastShadow = StaticMeshSection.bCastShadow;
-			SectionInfo.bEnableCollision = StaticMeshSection.bEnableCollision;
 			SectionInfo.MaterialIndex = StaticMeshSection.MaterialIndex;
 			StaticMesh->SectionInfoMap.Set(LodIndex, SectionIndex, SectionInfo);
 		}
@@ -416,7 +444,6 @@ void BuildVertexBuffer(
 			const TArray<FMeshTriangle>& PolygonTriangles = MeshDescription->GetPolygonTriangles(PolygonID);
 			int32 ReserveSize = IndexBuffer.Num() + PolygonTriangles.Num() * 3;
 			IndexBuffer.Reserve(ReserveSize);
-			OutWedgeMap.Reserve(ReserveSize);
 			uint32 MinIndex = TNumericLimits< uint32 >::Max();
 			uint32 MaxIndex = TNumericLimits< uint32 >::Min();
 			for (int32 TriangleIndex = 0; TriangleIndex < PolygonTriangles.Num(); ++TriangleIndex)
@@ -449,7 +476,7 @@ void BuildVertexBuffer(
 					StaticMeshVertex.TangentY = ScaleMatrix.TransformVector(FVector::CrossProduct(VertexInstanceNormal, VertexInstanceTangent).GetSafeNormal() * VertexInstanceBinormalSign).GetSafeNormal();
 					StaticMeshVertex.TangentZ = ScaleMatrix.TransformVector(VertexInstanceNormal).GetSafeNormal();
 					StaticMeshVertex.Color = LinearColor.ToFColor(true);
-					const uint32 MaxNumTexCoords = FMath::Min<int32>(MAX_MESH_TEXTURE_COORDS, MAX_STATIC_TEXCOORDS);
+					const uint32 MaxNumTexCoords = FMath::Min<int32>(MAX_MESH_TEXTURE_COORDS_MD, MAX_STATIC_TEXCOORDS);
 					for (uint32 UVIndex = 0; UVIndex < MaxNumTexCoords; ++UVIndex)
 					{
 						if(UVIndex < NumTextureCoord)
@@ -483,7 +510,7 @@ void BuildVertexBuffer(
 					RemapVerts[VertexInstanceValue] = Index;
 					const uint32 RenderingVertexIndex = RemapVerts[VertexInstanceValue];
 					IndexBuffer.Add(RenderingVertexIndex);
-					OutWedgeMap.Add(RenderingVertexIndex);
+					OutWedgeMap[VertexInstanceValue] = RenderingVertexIndex;
 					SectionIndices.Add(RenderingVertexIndex);
 					MinIndex = FMath::Min(MinIndex, RenderingVertexIndex);
 					MaxIndex = FMath::Max(MaxIndex, RenderingVertexIndex);
@@ -504,8 +531,7 @@ void BuildVertexBuffer(
 	//Optimize before setting the buffer
 	if (VertexInstances.Num() < 100000 * 3)
 	{
-		BuildOptimizationHelper::CacheOptimizeVertexAndIndexBuffer(StaticMeshBuildVertices, OutPerSectionIndices, OutWedgeMap);
-		check(OutWedgeMap.Num() == VertexInstances.Num());
+		BuildOptimizationHelper::CacheOptimizeVertexAndIndexBuffer(StaticMeshBuildVertices, OutPerSectionIndices);
 	}
 
 	StaticMeshLOD.VertexBuffer.SetUseHighPrecisionTangentBasis(LODBuildSettings.bUseHighPrecisionTangentBasis);
diff --git a/Engine/Source/Runtime/MeshBuilder/Public/MeshBuilder.h b/Engine/Source/Runtime/MeshBuilder/Public/MeshBuilder.h
--- a/Engine/Source/Runtime/MeshBuilder/Public/MeshBuilder.h
+++ b/Engine/Source/Runtime/MeshBuilder/Public/MeshBuilder.h
@@ -23,3 +23,12 @@ private:
 
 };
 
+class MESHBUILDER_API FMeshDescriptionOperations
+{
+public:
+	/** Convert this mesh description into the old FRawMesh format*/
+	static void ConverToRawMesh(const class UMeshDescription* SourceMeshDescription, struct FRawMesh &DestinationRawMesh);
+	/** Convert old FRawMesh format to MeshDescription*/
+	static void ConverFromRawMesh(const struct FRawMesh &SourceRawMesh, class UMeshDescription* DestinationMeshDescription);
+};
+
diff --git a/Engine/Source/Runtime/MeshBuilder/Public/StaticMeshBuilder.h b/Engine/Source/Runtime/MeshBuilder/Public/StaticMeshBuilder.h
--- a/Engine/Source/Runtime/MeshBuilder/Public/StaticMeshBuilder.h
+++ b/Engine/Source/Runtime/MeshBuilder/Public/StaticMeshBuilder.h
@@ -5,6 +5,8 @@
 #include "CoreMinimal.h"
 #include "MeshBuilder.h"
 
+DECLARE_LOG_CATEGORY_EXTERN(LogStaticMeshBuilder, Log, All);
+
 class MESHBUILDER_API FStaticMeshBuilder : public FMeshBuilder
 {
 public:
diff --git a/Engine/Source/Runtime/MeshDescription/Private/MeshDescription.cpp b/Engine/Source/Runtime/MeshDescription/Private/MeshDescription.cpp
--- a/Engine/Source/Runtime/MeshDescription/Private/MeshDescription.cpp
+++ b/Engine/Source/Runtime/MeshDescription/Private/MeshDescription.cpp
@@ -7,36 +7,71 @@
 #include "Serialization/ObjectWriter.h"
 
 UMeshDescription::UMeshDescription()
+{
+	RegisterBaseAttributes();
+}
+
+void UMeshDescription::RegisterBaseAttributes()
 {
 	// Add basic vertex attributes
-	VertexAttributes().RegisterAttribute<FVector>( MeshAttribute::Vertex::Position, 1, FVector::ZeroVector, EMeshAttributeFlags::Lerpable );
-	VertexAttributes().RegisterAttribute<float>( MeshAttribute::Vertex::CornerSharpness, 1, 0.0f, EMeshAttributeFlags::Lerpable );
+	VertexAttributes().RegisterAttribute<FVector>(MeshAttribute::Vertex::Position, 1, FVector::ZeroVector, EMeshAttributeFlags::Lerpable);
+	VertexAttributes().RegisterAttribute<float>(MeshAttribute::Vertex::CornerSharpness, 1, 0.0f, EMeshAttributeFlags::Lerpable);
 
 	// Add basic vertex instance attributes
-	VertexInstanceAttributes().RegisterAttribute<FVector2D>( MeshAttribute::VertexInstance::TextureCoordinate, 2, FVector2D::ZeroVector, EMeshAttributeFlags::Lerpable );
-	VertexInstanceAttributes().RegisterAttribute<FVector>( MeshAttribute::VertexInstance::Normal, 1, FVector::ZeroVector, EMeshAttributeFlags::AutoGenerated );
-	VertexInstanceAttributes().RegisterAttribute<FVector>( MeshAttribute::VertexInstance::Tangent, 1, FVector::ZeroVector, EMeshAttributeFlags::AutoGenerated );
-	VertexInstanceAttributes().RegisterAttribute<float>( MeshAttribute::VertexInstance::BinormalSign, 1, 0.0f, EMeshAttributeFlags::AutoGenerated );
-	VertexInstanceAttributes().RegisterAttribute<FVector4>( MeshAttribute::VertexInstance::Color, 1, FVector4( 1.0f ), EMeshAttributeFlags::Lerpable );
+	VertexInstanceAttributes().RegisterAttribute<FVector2D>(MeshAttribute::VertexInstance::TextureCoordinate, 2, FVector2D::ZeroVector, EMeshAttributeFlags::Lerpable);
+	VertexInstanceAttributes().RegisterAttribute<FVector>(MeshAttribute::VertexInstance::Normal, 1, FVector::ZeroVector, EMeshAttributeFlags::AutoGenerated);
+	VertexInstanceAttributes().RegisterAttribute<FVector>(MeshAttribute::VertexInstance::Tangent, 1, FVector::ZeroVector, EMeshAttributeFlags::AutoGenerated);
+	VertexInstanceAttributes().RegisterAttribute<float>(MeshAttribute::VertexInstance::BinormalSign, 1, 0.0f, EMeshAttributeFlags::AutoGenerated);
+	VertexInstanceAttributes().RegisterAttribute<FVector4>(MeshAttribute::VertexInstance::Color, 1, FVector4(1.0f), EMeshAttributeFlags::Lerpable);
 
 	// Add basic edge attributes
-	EdgeAttributes().RegisterAttribute<bool>( MeshAttribute::Edge::IsHard, 1, false );
-	EdgeAttributes().RegisterAttribute<float>( MeshAttribute::Edge::CreaseSharpness, 1, 0.0f, EMeshAttributeFlags::Lerpable );
+	EdgeAttributes().RegisterAttribute<bool>(MeshAttribute::Edge::IsHard, 1, false);
+	EdgeAttributes().RegisterAttribute<float>(MeshAttribute::Edge::CreaseSharpness, 1, 0.0f, EMeshAttributeFlags::Lerpable);
 
 	// Add basic polygon attributes
-	PolygonAttributes().RegisterAttribute<FVector>( MeshAttribute::Polygon::Normal, 1, FVector::ZeroVector, EMeshAttributeFlags::AutoGenerated );
-	PolygonAttributes().RegisterAttribute<FVector>( MeshAttribute::Polygon::Tangent, 1, FVector::ZeroVector, EMeshAttributeFlags::AutoGenerated );
-	PolygonAttributes().RegisterAttribute<FVector>( MeshAttribute::Polygon::Binormal, 1, FVector::ZeroVector, EMeshAttributeFlags::AutoGenerated );
-	PolygonAttributes().RegisterAttribute<FVector>( MeshAttribute::Polygon::Center, 1, FVector::ZeroVector, EMeshAttributeFlags::AutoGenerated );
+	PolygonAttributes().RegisterAttribute<FVector>(MeshAttribute::Polygon::Normal, 1, FVector::ZeroVector, EMeshAttributeFlags::AutoGenerated);
+	PolygonAttributes().RegisterAttribute<FVector>(MeshAttribute::Polygon::Tangent, 1, FVector::ZeroVector, EMeshAttributeFlags::AutoGenerated);
+	PolygonAttributes().RegisterAttribute<FVector>(MeshAttribute::Polygon::Binormal, 1, FVector::ZeroVector, EMeshAttributeFlags::AutoGenerated);
+	PolygonAttributes().RegisterAttribute<FVector>(MeshAttribute::Polygon::Center, 1, FVector::ZeroVector, EMeshAttributeFlags::AutoGenerated);
 
 	// Add basic polygon group attributes
-	PolygonGroupAttributes().RegisterAttribute<FSoftObjectPath>( MeshAttribute::PolygonGroup::MaterialAsset );
-	PolygonGroupAttributes().RegisterAttribute<FName>( MeshAttribute::PolygonGroup::MaterialSlotName );
-	PolygonGroupAttributes().RegisterAttribute<FName>( MeshAttribute::PolygonGroup::ImportedMaterialSlotName );
-	PolygonGroupAttributes().RegisterAttribute<bool>( MeshAttribute::PolygonGroup::EnableCollision );
-	PolygonGroupAttributes().RegisterAttribute<bool>( MeshAttribute::PolygonGroup::CastShadow );
+	PolygonGroupAttributes().RegisterAttribute<FSoftObjectPath>(MeshAttribute::PolygonGroup::MaterialAsset);
+	PolygonGroupAttributes().RegisterAttribute<FName>(MeshAttribute::PolygonGroup::MaterialSlotName);
+	PolygonGroupAttributes().RegisterAttribute<FName>(MeshAttribute::PolygonGroup::ImportedMaterialSlotName);
+	PolygonGroupAttributes().RegisterAttribute<bool>(MeshAttribute::PolygonGroup::EnableCollision);
+	PolygonGroupAttributes().RegisterAttribute<bool>(MeshAttribute::PolygonGroup::CastShadow);
 }
 
+void UMeshDescription::UnRegisterBaseAttributes()
+{
+	// Remove basic vertex attributes
+	VertexAttributes().UnregisterAttribute<FVector>(MeshAttribute::Vertex::Position);
+	VertexAttributes().UnregisterAttribute<float>(MeshAttribute::Vertex::CornerSharpness);
+
+	// Remove basic vertex instance attributes
+	VertexInstanceAttributes().UnregisterAttribute<FVector2D>(MeshAttribute::VertexInstance::TextureCoordinate);
+	VertexInstanceAttributes().UnregisterAttribute<FVector>(MeshAttribute::VertexInstance::Normal);
+	VertexInstanceAttributes().UnregisterAttribute<FVector>(MeshAttribute::VertexInstance::Tangent);
+	VertexInstanceAttributes().UnregisterAttribute<float>(MeshAttribute::VertexInstance::BinormalSign);
+	VertexInstanceAttributes().UnregisterAttribute<FVector4>(MeshAttribute::VertexInstance::Color);
+
+	// Remove basic edge attributes
+	EdgeAttributes().UnregisterAttribute<bool>(MeshAttribute::Edge::IsHard);
+	EdgeAttributes().UnregisterAttribute<float>(MeshAttribute::Edge::CreaseSharpness);
+
+	// Remove basic polygon attributes
+	PolygonAttributes().UnregisterAttribute<FVector>(MeshAttribute::Polygon::Normal);
+	PolygonAttributes().UnregisterAttribute<FVector>(MeshAttribute::Polygon::Tangent);
+	PolygonAttributes().UnregisterAttribute<FVector>(MeshAttribute::Polygon::Binormal);
+	PolygonAttributes().UnregisterAttribute<FVector>(MeshAttribute::Polygon::Center);
+
+	// Remove basic polygon group attributes
+	PolygonGroupAttributes().UnregisterAttribute<FSoftObjectPath>(MeshAttribute::PolygonGroup::MaterialAsset);
+	PolygonGroupAttributes().UnregisterAttribute<FName>(MeshAttribute::PolygonGroup::MaterialSlotName);
+	PolygonGroupAttributes().UnregisterAttribute<FName>(MeshAttribute::PolygonGroup::ImportedMaterialSlotName);
+	PolygonGroupAttributes().UnregisterAttribute<bool>(MeshAttribute::PolygonGroup::EnableCollision);
+	PolygonGroupAttributes().UnregisterAttribute<bool>(MeshAttribute::PolygonGroup::CastShadow);
+}
 
 void UMeshDescription::Serialize( FArchive& Ar )
 {
@@ -55,6 +90,18 @@ void UMeshDescription::Serialize( FArchive& Ar )
 	Ar << PolygonGroupAttributesSet;
 }
 
+void UMeshDescription::Empty()
+{
+	VertexArray.Reset();
+	VertexInstanceArray.Reset();
+	EdgeArray.Reset();
+	PolygonArray.Reset();
+	PolygonGroupArray.Reset();
+
+	//Empty all attributes
+	UnRegisterBaseAttributes();
+	RegisterBaseAttributes();
+}
 
 #if WITH_EDITORONLY_DATA
 FString UMeshDescription::GetIdString()
@@ -207,3 +254,254 @@ void UMeshDescription::RemapAttributes( const FElementIDRemappings& Remappings )
 	PolygonAttributesSet.Remap( Remappings.NewPolygonIndexLookup );
 	PolygonGroupAttributesSet.Remap( Remappings.NewPolygonGroupIndexLookup );
 }
+
+void UMeshDescription::GetPolygonPerimeterVertices(const FPolygonID PolygonID, TArray<FVertexID>& OutPolygonPerimeterVertexIDs) const
+{
+	const FMeshPolygon& Polygon = GetPolygon(PolygonID);
+
+	OutPolygonPerimeterVertexIDs.SetNumUninitialized(Polygon.PerimeterContour.VertexInstanceIDs.Num(), false);
+
+	int32 Index = 0;
+	for (const FVertexInstanceID VertexInstanceID : Polygon.PerimeterContour.VertexInstanceIDs)
+	{
+		const FMeshVertexInstance& VertexInstance = VertexInstanceArray[VertexInstanceID];
+		OutPolygonPerimeterVertexIDs[Index] = VertexInstance.VertexID;
+		Index++;
+	}
+}
+
+/** Given three direction vectors, indicates if A and B are on the same 'side' of Vec. */
+bool UMeshDescription::VectorsOnSameSide(const FVector& Vec, const FVector& A, const FVector& B, const float SameSideDotProductEpsilon)
+{
+	const FVector CrossA = Vec ^ A;
+	const FVector CrossB = Vec ^ B;
+	float DotWithEpsilon = SameSideDotProductEpsilon + (CrossA | CrossB);
+	return !FMath::IsNegativeFloat(DotWithEpsilon);
+}
+
+/** Util to see if P lies within triangle created by A, B and C. */
+bool UMeshDescription::PointInTriangle(const FVector& A, const FVector& B, const FVector& C, const FVector& P, const float InsideTriangleDotProductEpsilon)
+{
+	// Cross product indicates which 'side' of the vector the point is on
+	// If its on the same side as the remaining vert for all edges, then its inside.	
+	if (VectorsOnSameSide(B - A, P - A, C - A, InsideTriangleDotProductEpsilon) &&
+		VectorsOnSameSide(C - B, P - B, A - B, InsideTriangleDotProductEpsilon) &&
+		VectorsOnSameSide(A - C, P - C, B - C, InsideTriangleDotProductEpsilon))
+	{
+		return true;
+	}
+	else
+	{
+		return false;
+	}
+
+}
+
+FPlane UMeshDescription::ComputePolygonPlane(const FPolygonID PolygonID) const
+{
+	// NOTE: This polygon plane computation code is partially based on the implementation of "Newell's method" from Real-Time 
+	//       Collision Detection by Christer Ericson, published by Morgan Kaufmann Publishers, (c) 2005 Elsevier Inc
+
+	// @todo mesheditor perf: For polygons that are just triangles, use a cross product to get the normal fast!
+	// @todo mesheditor perf: We could skip computing the plane distance when we only need the normal
+	// @todo mesheditor perf: We could cache these computed polygon normals; or just use the normal of the first three vertices' triangle if it is satisfactory in all cases
+	// @todo mesheditor: For non-planar polygons, the result can vary. Ideally this should use the actual polygon triangulation as opposed to the arbitrary triangulation used here.
+
+	FVector Centroid = FVector::ZeroVector;
+	FVector Normal = FVector::ZeroVector;
+
+	static TArray<FVertexID> PerimeterVertexIDs;
+	GetPolygonPerimeterVertices(PolygonID, /* Out */ PerimeterVertexIDs);
+
+	const TVertexAttributeArray<FVector>& VertexPositions = VertexAttributes().GetAttributes<FVector>(MeshAttribute::Vertex::Position);
+
+	// Use 'Newell's Method' to compute a robust 'best fit' plane from the vertices of this polygon
+	for (int32 VertexNumberI = PerimeterVertexIDs.Num() - 1, VertexNumberJ = 0; VertexNumberJ < PerimeterVertexIDs.Num(); VertexNumberI = VertexNumberJ, VertexNumberJ++)
+	{
+		const FVertexID VertexIDI = PerimeterVertexIDs[VertexNumberI];
+		const FVector PositionI = VertexPositions[VertexIDI];
+
+		const FVertexID VertexIDJ = PerimeterVertexIDs[VertexNumberJ];
+		const FVector PositionJ = VertexPositions[VertexIDJ];
+
+		Centroid += PositionJ;
+
+		Normal.X += (PositionJ.Y - PositionI.Y) * (PositionI.Z + PositionJ.Z);
+		Normal.Y += (PositionJ.Z - PositionI.Z) * (PositionI.X + PositionJ.X);
+		Normal.Z += (PositionJ.X - PositionI.X) * (PositionI.Y + PositionJ.Y);
+	}
+
+	Normal.Normalize();
+
+	// Construct a plane from the normal and centroid
+	return FPlane(Normal, FVector::DotProduct(Centroid, Normal) / (float)PerimeterVertexIDs.Num());
+}
+
+
+FVector UMeshDescription::ComputePolygonNormal(const FPolygonID PolygonID) const
+{
+	// @todo mesheditor: Polygon normals are now computed and cached when changes are made to a polygon.
+	// In theory, we can just return that cached value, but we need to check that there is nothing which relies on the value being correct before
+	// the cache is updated at the end of a modification.
+	const FPlane PolygonPlane = ComputePolygonPlane(PolygonID);
+	const FVector PolygonNormal(PolygonPlane.X, PolygonPlane.Y, PolygonPlane.Z);
+	return PolygonNormal;
+}
+
+void UMeshDescription::ComputePolygonTriangulation(const FPolygonID PolygonID, TArray<FMeshTriangle>& OutTriangles)
+{
+	// NOTE: This polygon triangulation code is partially based on the ear cutting algorithm described on
+	//       page 497 of the book "Real-time Collision Detection", published in 2005.
+
+	struct Local
+	{
+		// Returns true if the triangle formed by the specified three positions has a normal that is facing the opposite direction of the reference normal
+		static inline bool IsTriangleFlipped(const FVector ReferenceNormal, const FVector VertexPositionA, const FVector VertexPositionB, const FVector VertexPositionC)
+		{
+			const FVector TriangleNormal = FVector::CrossProduct(
+				VertexPositionC - VertexPositionA,
+				VertexPositionB - VertexPositionA).GetSafeNormal();
+			return (FVector::DotProduct(ReferenceNormal, TriangleNormal) <= 0.0f);
+		}
+
+	};
+
+
+	OutTriangles.Reset();
+
+	// @todo mesheditor holes: Does not support triangles with holes yet!
+	// @todo mesheditor: Perhaps should always attempt to triangulate by splitting polygons along the shortest edge, for better determinism.
+
+	//	const FMeshPolygon& Polygon = GetPolygon( PolygonID );
+	const TArray<FVertexInstanceID>& PolygonVertexInstanceIDs = GetPolygonPerimeterVertexInstances(PolygonID);
+
+	// Polygon must have at least three vertices/edges
+	const int32 PolygonVertexCount = PolygonVertexInstanceIDs.Num();
+	check(PolygonVertexCount >= 3);
+
+	// First figure out the polygon normal.  We need this to determine which triangles are convex, so that
+	// we can figure out which ears to clip
+	const FVector PolygonNormal = ComputePolygonNormal(PolygonID);
+
+	// Make a simple linked list array of the previous and next vertex numbers, for each vertex number
+	// in the polygon.  This will just save us having to iterate later on.
+	static TArray<int32> PrevVertexNumbers, NextVertexNumbers;
+	static TArray<FVector> VertexPositions;
+
+	{
+		const TVertexAttributeArray<FVector>& MeshVertexPositions = VertexAttributes().GetAttributes<FVector>(MeshAttribute::Vertex::Position);
+		PrevVertexNumbers.SetNumUninitialized(PolygonVertexCount, false);
+		NextVertexNumbers.SetNumUninitialized(PolygonVertexCount, false);
+		VertexPositions.SetNumUninitialized(PolygonVertexCount, false);
+
+		for (int32 VertexNumber = 0; VertexNumber < PolygonVertexCount; ++VertexNumber)
+		{
+			PrevVertexNumbers[VertexNumber] = VertexNumber - 1;
+			NextVertexNumbers[VertexNumber] = VertexNumber + 1;
+
+			VertexPositions[VertexNumber] = MeshVertexPositions[GetVertexInstanceVertex(PolygonVertexInstanceIDs[VertexNumber])];
+		}
+		PrevVertexNumbers[0] = PolygonVertexCount - 1;
+		NextVertexNumbers[PolygonVertexCount - 1] = 0;
+	}
+
+	int32 EarVertexNumber = 0;
+	int32 EarTestCount = 0;
+	for (int32 RemainingVertexCount = PolygonVertexCount; RemainingVertexCount >= 3; )
+	{
+		bool bIsEar = true;
+
+		// If we're down to only a triangle, just treat it as an ear.  Also, if we've tried every possible candidate
+		// vertex looking for an ear, go ahead and just treat the current vertex as an ear.  This can happen when 
+		// vertices are colinear or other degenerate cases.
+		if (RemainingVertexCount > 3 && EarTestCount < RemainingVertexCount)
+		{
+			const FVector PrevVertexPosition = VertexPositions[PrevVertexNumbers[EarVertexNumber]];
+			const FVector EarVertexPosition = VertexPositions[EarVertexNumber];
+			const FVector NextVertexPosition = VertexPositions[NextVertexNumbers[EarVertexNumber]];
+
+			// Figure out whether the potential ear triangle is facing the same direction as the polygon
+			// itself.  If it's facing the opposite direction, then we're dealing with a concave triangle
+			// and we'll skip it for now.
+			if (!Local::IsTriangleFlipped(
+				PolygonNormal,
+				PrevVertexPosition,
+				EarVertexPosition,
+				NextVertexPosition))
+			{
+				int32 TestVertexNumber = NextVertexNumbers[NextVertexNumbers[EarVertexNumber]];
+
+				do
+				{
+					// Test every other remaining vertex to make sure that it doesn't lie inside our potential ear
+					// triangle.  If we find a vertex that's inside the triangle, then it cannot actually be an ear.
+					const FVector TestVertexPosition = VertexPositions[TestVertexNumber];
+					if (PointInTriangle(
+						PrevVertexPosition,
+						EarVertexPosition,
+						NextVertexPosition,
+						TestVertexPosition,
+						SMALL_NUMBER))
+					{
+						bIsEar = false;
+						break;
+					}
+
+					TestVertexNumber = NextVertexNumbers[TestVertexNumber];
+				} while (TestVertexNumber != PrevVertexNumbers[EarVertexNumber]);
+			}
+			else
+			{
+				bIsEar = false;
+			}
+		}
+
+		if (bIsEar)
+		{
+			// OK, we found an ear!  Let's save this triangle in our output buffer.
+			{
+				OutTriangles.Emplace();
+				FMeshTriangle& Triangle = OutTriangles.Last();
+
+				Triangle.SetVertexInstanceID(0, PolygonVertexInstanceIDs[PrevVertexNumbers[EarVertexNumber]]);
+				Triangle.SetVertexInstanceID(1, PolygonVertexInstanceIDs[EarVertexNumber]);
+				Triangle.SetVertexInstanceID(2, PolygonVertexInstanceIDs[NextVertexNumbers[EarVertexNumber]]);
+			}
+
+			// Update our linked list.  We're effectively cutting off the ear by pointing the ear vertex's neighbors to
+			// point at their next sequential neighbor, and reducing the remaining vertex count by one.
+			{
+				NextVertexNumbers[PrevVertexNumbers[EarVertexNumber]] = NextVertexNumbers[EarVertexNumber];
+				PrevVertexNumbers[NextVertexNumbers[EarVertexNumber]] = PrevVertexNumbers[EarVertexNumber];
+				--RemainingVertexCount;
+			}
+
+			// Move on to the previous vertex in the list, now that this vertex was cut
+			EarVertexNumber = PrevVertexNumbers[EarVertexNumber];
+
+			EarTestCount = 0;
+		}
+		else
+		{
+			// The vertex is not the ear vertex, because it formed a triangle that either had a normal which pointed in the opposite direction
+			// of the polygon, or at least one of the other polygon vertices was found to be inside the triangle.  Move on to the next vertex.
+			EarVertexNumber = NextVertexNumbers[EarVertexNumber];
+
+			// Keep track of how many ear vertices we've tested, so that if we exhaust all remaining vertices, we can
+			// fall back to clipping the triangle and adding it to our mesh anyway.  This is important for degenerate cases.
+			++EarTestCount;
+		}
+	}
+
+	check(OutTriangles.Num() > 0);
+}
+
+void UMeshDescription::TriangulateMesh()
+{
+	// Perform triangulation directly into mesh polygons
+	for (const FPolygonID PolygonID : Polygons().GetElementIDs())
+	{
+		FMeshPolygon& Polygon = PolygonArray[PolygonID];
+		ComputePolygonTriangulation(PolygonID, Polygon.Triangles);
+	}
+}
\ No newline at end of file
diff --git a/Engine/Source/Runtime/MeshDescription/Public/MeshDescription.h b/Engine/Source/Runtime/MeshDescription/Public/MeshDescription.h
--- a/Engine/Source/Runtime/MeshDescription/Public/MeshDescription.h
+++ b/Engine/Source/Runtime/MeshDescription/Public/MeshDescription.h
@@ -242,6 +242,9 @@ public:
 	// UObject interface
 	virtual void Serialize( FArchive& Ar ) override;
 
+	//Empty the meshdescription
+	void Empty();
+
 	FVertexArray& Vertices() { return VertexArray; }
 	const FVertexArray& Vertices() const { return VertexArray; }
 
@@ -481,7 +484,7 @@ private:
 		}
 
 		Polygon.PolygonGroupID = PolygonGroupID;
-		check( !PolygonGroupArray[ PolygonGroupID ].Polygons.Contains( PolygonID ) );
+		checkSlow( !PolygonGroupArray[ PolygonGroupID ].Polygons.Contains( PolygonID ) );
 		PolygonGroupArray[ PolygonGroupID ].Polygons.Add( PolygonID );
 
 		PolygonAttributesSet.Insert( PolygonID );
@@ -670,6 +673,17 @@ public:
 		return PolygonArray[ PolygonID ].PerimeterContour.VertexInstanceIDs;
 	}
 
+	void GetPolygonEdges(const FPolygonID PolygonID, TArray<FEdgeID>& OutPolygonEdges) const
+	{
+		const FMeshPolygonContour& PerimeterContour = PolygonArray[PolygonID].PerimeterContour;
+		const int32 ContourCount = PerimeterContour.VertexInstanceIDs.Num();
+		for (int32 ContourIndex = 0; ContourIndex < ContourCount; ++ContourIndex)
+		{
+			int32 ContourPlusOne = (ContourIndex + 1) % ContourCount;
+			OutPolygonEdges.Add(GetVertexPairEdge(GetVertexInstanceVertex(PerimeterContour.VertexInstanceIDs[ContourIndex]), GetVertexInstanceVertex(PerimeterContour.VertexInstanceIDs[ContourPlusOne])));
+		}
+	}
+
 	/** Return the polygon group associated with a polygon */
 	FPolygonGroupID GetPolygonPolygonGroup( const FPolygonID PolygonID ) const
 	{
@@ -726,6 +740,17 @@ public:
 	FString GetIdString();
 #endif
 
+public:
+	void ComputePolygonTriangulation(const FPolygonID PolygonID, TArray<FMeshTriangle>& OutTriangles);
+	void TriangulateMesh();
+
+private:
+	bool VectorsOnSameSide(const FVector& Vec, const FVector& A, const FVector& B, const float SameSideDotProductEpsilon);
+	bool PointInTriangle(const FVector& A, const FVector& B, const FVector& C, const FVector& P, const float InsideTriangleDotProductEpsilon);
+	void GetPolygonPerimeterVertices(const FPolygonID PolygonID, TArray<FVertexID>& OutPolygonPerimeterVertexIDs) const;
+	FPlane ComputePolygonPlane(const FPolygonID PolygonID) const;
+	FVector ComputePolygonNormal(const FPolygonID PolygonID) const;
+
 private:
 
 	/** Given a set of index remappings, fixes up references to element IDs */
@@ -734,6 +759,9 @@ private:
 	/** Given a set of index remappings, remaps all attributes accordingly */
 	void RemapAttributes( const FElementIDRemappings& Remappings );
 
+	void RegisterBaseAttributes();
+	void UnRegisterBaseAttributes();
+
 	FVertexArray VertexArray;
 	FVertexInstanceArray VertexInstanceArray;
 	FEdgeArray EdgeArray;
