Add NTBs calculation to the meshbuilder API

diff --git a/Engine/Source/Runtime/MeshBuilder/Private/MeshBuilder.cpp b/Engine/Source/Runtime/MeshBuilder/Private/MeshBuilder.cpp
--- a/Engine/Source/Runtime/MeshBuilder/Private/MeshBuilder.cpp
+++ b/Engine/Source/Runtime/MeshBuilder/Private/MeshBuilder.cpp
@@ -3,6 +3,7 @@
 #include "MeshBuilder.h"
 #include "MeshDescriptionHelper.h"
 #include "MeshDescription.h"
+#include "MeshAttributes.h"
 #include "RawMesh.h"
 
 
@@ -11,6 +12,64 @@ FMeshBuilder::FMeshBuilder()
 
 }
 
+void FMeshDescriptionOperations::ComputeMeshNTBs(class UMeshDescription* MeshDescription, const struct FMeshBuildSettings& BuildSettings)
+{
+	FVertexInstanceArray& VertexInstanceArray = MeshDescription->VertexInstances();
+	TVertexInstanceAttributeArray<FVector>& Normals = MeshDescription->VertexInstanceAttributes().GetAttributes<FVector>(MeshAttribute::VertexInstance::Normal);
+	TVertexInstanceAttributeArray<FVector>& Tangents = MeshDescription->VertexInstanceAttributes().GetAttributes<FVector>(MeshAttribute::VertexInstance::Tangent);
+	TVertexInstanceAttributeArray<float>& BinormalSigns = MeshDescription->VertexInstanceAttributes().GetAttributes<float>(MeshAttribute::VertexInstance::BinormalSign);
+
+	// Static meshes always blend normals of overlapping corners.
+	uint32 TangentOptions = FMeshDescriptionHelper::ETangentOptions::BlendOverlappingNormals;
+	if (BuildSettings.bRemoveDegenerates)
+	{
+		// If removing degenerate triangles, ignore them when computing tangents.
+		TangentOptions |= FMeshDescriptionHelper::ETangentOptions::IgnoreDegenerateTriangles;
+	}
+
+	//This function make sure the Polygon NTB are compute and also remove degenerated triangle from the render mesh description.
+	FMeshDescriptionHelper::CreatePolygonNTB(MeshDescription, BuildSettings.bRemoveDegenerates ? SMALL_NUMBER : 0.0f);
+
+	//Keep the original mesh description NTBs if we do not rebuild the normals or tangents.
+	bool bComputeTangentLegacy = !BuildSettings.bUseMikkTSpace && (BuildSettings.bRecomputeNormals || BuildSettings.bRecomputeTangents);
+	bool bHasAllNormals = true;
+	bool bHasAllTangents = true;
+	for (const FVertexInstanceID VertexInstanceID : VertexInstanceArray.GetElementIDs())
+	{
+		// Dump normals and tangents if we are recomputing them.
+		if (BuildSettings.bRecomputeTangents)
+		{
+			//Dump the tangents
+			BinormalSigns[VertexInstanceID] = 0.0f;
+			Tangents[VertexInstanceID] = FVector(0.0f);
+		}
+		if (BuildSettings.bRecomputeNormals)
+		{
+			//Dump the normals
+			Normals[VertexInstanceID] = FVector(0.0f);
+		}
+		bHasAllNormals &= !Normals[VertexInstanceID].IsNearlyZero();
+		bHasAllTangents &= !Tangents[VertexInstanceID].IsNearlyZero();
+	}
+
+
+	//MikkTSpace should be use only when the user want to recompute the normals or tangents otherwise should always fallback on builtin
+	//We cannot use mikkt space with degenerated normals fallback on buitin.
+	if (BuildSettings.bUseMikkTSpace && (BuildSettings.bRecomputeNormals || BuildSettings.bRecomputeTangents))
+	{
+		if (!bHasAllNormals)
+		{
+			FMeshDescriptionHelper::CreateNormals(MeshDescription, (FMeshDescriptionHelper::ETangentOptions)TangentOptions, false);
+		}
+		FMeshDescriptionHelper::CreateMikktTangents(MeshDescription, (FMeshDescriptionHelper::ETangentOptions)TangentOptions);
+	}
+	else if (!bHasAllNormals || !bHasAllTangents)
+	{
+		//Set the compute tangent to true when we do not build using mikkt space
+		FMeshDescriptionHelper::CreateNormals(MeshDescription, (FMeshDescriptionHelper::ETangentOptions)TangentOptions, true);
+	}
+}
+
 /** Convert this mesh description into the old FRawMesh format*/
 void FMeshDescriptionOperations::ConverToRawMesh(const UMeshDescription* SourceMeshDescription, FRawMesh &DestinationRawMesh)
 {
diff --git a/Engine/Source/Runtime/MeshBuilder/Private/MeshDescriptionHelper.h b/Engine/Source/Runtime/MeshBuilder/Private/MeshDescriptionHelper.h
--- a/Engine/Source/Runtime/MeshBuilder/Private/MeshDescriptionHelper.h
+++ b/Engine/Source/Runtime/MeshBuilder/Private/MeshDescriptionHelper.h
@@ -57,6 +57,9 @@ public:
 	/** Convert old FRawMesh format to MeshDescription*/
 	static void ConverFromRawMesh(const struct FRawMesh &SourceRawMesh, UMeshDescription* DestinationMeshDescription);
 
+	static void CreatePolygonNTB(UMeshDescription* MeshDescription, float ComparisonThreshold);
+	static void CreateNormals(UMeshDescription* MeshDescription, ETangentOptions TangentOptions, bool bComputeTangent);
+	static void CreateMikktTangents(UMeshDescription* MeshDescription, ETangentOptions TangentOptions);
 private:
 
 	//////////////////////////////////////////////////////////////////////////
@@ -65,10 +68,6 @@ private:
 	static void ConvertHardEdgesToSmoothGroup(const UMeshDescription* SourceMeshDescription, struct FRawMesh &DestinationRawMesh);
 	static void ConvertSmoothGroupToHardEdges(const struct FRawMesh &SourceRawMesh, UMeshDescription* DestinationMeshDescription);
 
-	static void CreateNormals(UMeshDescription* MeshDescription, ETangentOptions TangentOptions, bool bComputeTangent);
-	static void CreateMikktTangents(UMeshDescription* MeshDescription, ETangentOptions TangentOptions);
-	static void CreatePolygonNTB(UMeshDescription* MeshDescription, float ComparisonThreshold);
-
 	//////////////////////////////////////////////////////////////////////////
 	//PRIVATE class members
 
diff --git a/Engine/Source/Runtime/MeshBuilder/Public/MeshBuilder.h b/Engine/Source/Runtime/MeshBuilder/Public/MeshBuilder.h
--- a/Engine/Source/Runtime/MeshBuilder/Public/MeshBuilder.h
+++ b/Engine/Source/Runtime/MeshBuilder/Public/MeshBuilder.h
@@ -26,8 +26,12 @@ private:
 class MESHBUILDER_API FMeshDescriptionOperations
 {
 public:
+	/** Compute the normals tangents and binormals for the mesh description */
+	static void ComputeMeshNTBs(class UMeshDescription* MeshDescription, const struct FMeshBuildSettings& BuildSettings);
+
 	/** Convert this mesh description into the old FRawMesh format*/
 	static void ConverToRawMesh(const class UMeshDescription* SourceMeshDescription, struct FRawMesh &DestinationRawMesh);
+
 	/** Convert old FRawMesh format to MeshDescription*/
 	static void ConverFromRawMesh(const struct FRawMesh &SourceRawMesh, class UMeshDescription* DestinationMeshDescription);
 };
