Add material index attribute to polygon group attributes container

diff --git a/Engine/Plugins/Tests/EditorTests/Source/EditorTests/Private/UnrealEd/MeshDescriptionTests.cpp b/Engine/Plugins/Tests/EditorTests/Source/EditorTests/Private/UnrealEd/MeshDescriptionTests.cpp
--- a/Engine/Plugins/Tests/EditorTests/Source/EditorTests/Private/UnrealEd/MeshDescriptionTests.cpp
+++ b/Engine/Plugins/Tests/EditorTests/Source/EditorTests/Private/UnrealEd/MeshDescriptionTests.cpp
@@ -132,7 +132,6 @@ bool FMeshDescriptionTest::CompareMeshDescription(const FString& AssetName, FAut
 	//////////////////////////////////////////////////////////////////////////
 	//Gather the reference data
 	const TVertexAttributeArray<FVector>& ReferenceVertexPositions = ReferenceMeshDescription->VertexAttributes().GetAttributes<FVector>(MeshAttribute::Vertex::Position);
-	//const TVertexAttributeArray<float>& ReferenceVertexCornerSharpness = ReferenceMeshDescription->VertexAttributes().GetAttributes<float>(MeshAttribute::Vertex::CornerSharpness);
 
 	const TVertexInstanceAttributeArray<FVector>& ReferenceVertexInstanceNormals = ReferenceMeshDescription->VertexInstanceAttributes().GetAttributes<FVector>(MeshAttribute::VertexInstance::Normal);
 	const TVertexInstanceAttributeArray<FVector>& ReferenceVertexInstanceTangents = ReferenceMeshDescription->VertexInstanceAttributes().GetAttributes<FVector>(MeshAttribute::VertexInstance::Tangent);
@@ -141,24 +140,13 @@ bool FMeshDescriptionTest::CompareMeshDescription(const FString& AssetName, FAut
 	const TVertexInstanceAttributeIndicesArray<FVector2D>& ReferenceVertexInstanceUVs = ReferenceMeshDescription->VertexInstanceAttributes().GetAttributesSet<FVector2D>(MeshAttribute::VertexInstance::TextureCoordinate);
 
 	const TEdgeAttributeArray<bool>& ReferenceEdgeHardnesses = ReferenceMeshDescription->EdgeAttributes().GetAttributes<bool>(MeshAttribute::Edge::IsHard);
-	//const TEdgeAttributeArray<float>& ReferenceEdgeCreaseSharpnesses = ReferenceMeshDescription->EdgeAttributes().GetAttributes<float>(MeshAttribute::Edge::CreaseSharpness);
 
-	//const TPolygonAttributeArray<FVector>& ReferencePolygonNormals = ReferenceMeshDescription->PolygonAttributes().GetAttributes<FVector>(MeshAttribute::Polygon::Normal);
-	//const TPolygonAttributeArray<FVector>& ReferencePolygonTangents = ReferenceMeshDescription->PolygonAttributes().GetAttributes<FVector>(MeshAttribute::Polygon::Tangent);
-	//const TPolygonAttributeArray<FVector>& ReferencePolygonBiNormals = ReferenceMeshDescription->PolygonAttributes().GetAttributes<FVector>(MeshAttribute::Polygon::Binormal);
-	//const TPolygonAttributeArray<FVector>& ReferencePolygonCenters = ReferenceMeshDescription->PolygonAttributes().GetAttributes<FVector>(MeshAttribute::Polygon::Center);
-
-	const TPolygonGroupAttributeArray<FSoftObjectPath>& ReferencePolygonGroupMaterialAssets = ReferenceMeshDescription->PolygonGroupAttributes().GetAttributes<FSoftObjectPath>(MeshAttribute::PolygonGroup::MaterialAsset);
-	const TPolygonGroupAttributeArray<FName>& ReferencePolygonGroupMaterialSlotNames = ReferenceMeshDescription->PolygonGroupAttributes().GetAttributes<FName>(MeshAttribute::PolygonGroup::MaterialSlotName);
-	const TPolygonGroupAttributeArray<FName>& ReferencePolygonGroupImportedMaterialSlotNames = ReferenceMeshDescription->PolygonGroupAttributes().GetAttributes<FName>(MeshAttribute::PolygonGroup::ImportedMaterialSlotName);
-	const TPolygonGroupAttributeArray<bool>& ReferencePolygonGroupCollision = ReferenceMeshDescription->PolygonGroupAttributes().GetAttributes<bool>(MeshAttribute::PolygonGroup::EnableCollision);
-	const TPolygonGroupAttributeArray<bool>& ReferencePolygonGroupCastShadow = ReferenceMeshDescription->PolygonGroupAttributes().GetAttributes<bool>(MeshAttribute::PolygonGroup::CastShadow);
+	const TPolygonGroupAttributeArray<int>& ReferencePolygonGroupMaterialIndex = ReferenceMeshDescription->PolygonGroupAttributes().GetAttributes<int>(MeshAttribute::PolygonGroup::MaterialIndex);
 
 
 	//////////////////////////////////////////////////////////////////////////
 	//Gather the result data
 	const TVertexAttributeArray<FVector>& ResultVertexPositions = MeshDescription->VertexAttributes().GetAttributes<FVector>(MeshAttribute::Vertex::Position);
-	//const TVertexAttributeArray<float>& ResultVertexCornerSharpness = MeshDescription->VertexAttributes().GetAttributes<float>(MeshAttribute::Vertex::CornerSharpness);
 
 	const TVertexInstanceAttributeArray<FVector>& ResultVertexInstanceNormals = MeshDescription->VertexInstanceAttributes().GetAttributes<FVector>(MeshAttribute::VertexInstance::Normal);
 	const TVertexInstanceAttributeArray<FVector>& ResultVertexInstanceTangents = MeshDescription->VertexInstanceAttributes().GetAttributes<FVector>(MeshAttribute::VertexInstance::Tangent);
@@ -167,18 +155,8 @@ bool FMeshDescriptionTest::CompareMeshDescription(const FString& AssetName, FAut
 	const TVertexInstanceAttributeIndicesArray<FVector2D>& ResultVertexInstanceUVs = MeshDescription->VertexInstanceAttributes().GetAttributesSet<FVector2D>(MeshAttribute::VertexInstance::TextureCoordinate);
 
 	const TEdgeAttributeArray<bool>& ResultEdgeHardnesses = MeshDescription->EdgeAttributes().GetAttributes<bool>(MeshAttribute::Edge::IsHard);
-	//const TEdgeAttributeArray<float>& ResultEdgeCreaseSharpnesses = MeshDescription->EdgeAttributes().GetAttributes<float>(MeshAttribute::Edge::CreaseSharpness);
-
-	//const TPolygonAttributeArray<FVector>& ResultPolygonNormals = MeshDescription->PolygonAttributes().GetAttributes<FVector>(MeshAttribute::Polygon::Normal);
-	//const TPolygonAttributeArray<FVector>& ResultPolygonTangents = MeshDescription->PolygonAttributes().GetAttributes<FVector>(MeshAttribute::Polygon::Tangent);
-	//const TPolygonAttributeArray<FVector>& ResultPolygonBiNormals = MeshDescription->PolygonAttributes().GetAttributes<FVector>(MeshAttribute::Polygon::Binormal);
-	//const TPolygonAttributeArray<FVector>& ResultPolygonCenters = MeshDescription->PolygonAttributes().GetAttributes<FVector>(MeshAttribute::Polygon::Center);
 
-	const TPolygonGroupAttributeArray<FSoftObjectPath>& ResultPolygonGroupMaterialAssets = MeshDescription->PolygonGroupAttributes().GetAttributes<FSoftObjectPath>(MeshAttribute::PolygonGroup::MaterialAsset);
-	const TPolygonGroupAttributeArray<FName>& ResultPolygonGroupMaterialSlotNames = MeshDescription->PolygonGroupAttributes().GetAttributes<FName>(MeshAttribute::PolygonGroup::MaterialSlotName);
-	const TPolygonGroupAttributeArray<FName>& ResultPolygonGroupImportedMaterialSlotNames = MeshDescription->PolygonGroupAttributes().GetAttributes<FName>(MeshAttribute::PolygonGroup::ImportedMaterialSlotName);
-	const TPolygonGroupAttributeArray<bool>& ResultPolygonGroupCollision = MeshDescription->PolygonGroupAttributes().GetAttributes<bool>(MeshAttribute::PolygonGroup::EnableCollision);
-	const TPolygonGroupAttributeArray<bool>& ResultPolygonGroupCastShadow = MeshDescription->PolygonGroupAttributes().GetAttributes<bool>(MeshAttribute::PolygonGroup::CastShadow);
+	const TPolygonGroupAttributeArray<int>& ResultPolygonGroupMaterialIndex = MeshDescription->PolygonGroupAttributes().GetAttributes<int>(MeshAttribute::PolygonGroup::MaterialIndex);
 
 	
 	//////////////////////////////////////////////////////////////////////////
@@ -352,16 +330,6 @@ bool FMeshDescriptionTest::CompareMeshDescription(const FString& AssetName, FAut
 			}
 		}
 	}
-	//Polygon group ID
-	//We currently rely only on the PolygonGroupID. The duplicate material slot information (the info is store in the staticmesh material array) is not necessary and cannot be put in FRawMesh structure.
-	if (ReferenceMeshDescription->PolygonGroups().Num() != MeshDescription->PolygonGroups().Num())
-	{
-		ExecutionInfo.AddEvent(FAutomationEvent(EAutomationEventType::Error, FString::Printf(TEXT("The %s conversion to RawMesh is not lossless, Polygon group count is different. Polygon group count expected [%d] result [%d]"),
-			*AssetName,
-			ReferenceMeshDescription->PolygonGroups().Num(),
-			MeshDescription->PolygonGroups().Num())));
-		bAllSame = false;
-	}
 	//Edges
 	//We check if hard edges are kept correctly
 	if (ReferenceEdgeHardnesses.Num() != ResultEdgeHardnesses.Num())
@@ -394,7 +362,32 @@ bool FMeshDescriptionTest::CompareMeshDescription(const FString& AssetName, FAut
 			}
 		}
 	}
-
+	//Polygon group ID
+	//We currently rely only on the PolygonGroupID. The duplicate material slot information (the info is store in the staticmesh material array) is not necessary and cannot be put in FRawMesh structure.
+	if (ReferenceMeshDescription->PolygonGroups().Num() != MeshDescription->PolygonGroups().Num())
+	{
+		ExecutionInfo.AddEvent(FAutomationEvent(EAutomationEventType::Error, FString::Printf(TEXT("The %s conversion to RawMesh is not lossless, Polygon group count is different. Polygon group count expected [%d] result [%d]"),
+			*AssetName,
+			ReferenceMeshDescription->PolygonGroups().Num(),
+			MeshDescription->PolygonGroups().Num())));
+		bAllSame = false;
+	}
+	else
+	{
+		for (const FPolygonGroupID& PolygonGroupID : ReferenceMeshDescription->PolygonGroups().GetElementIDs())
+		{
+			if (ReferencePolygonGroupMaterialIndex[PolygonGroupID] != ResultPolygonGroupMaterialIndex[PolygonGroupID])
+			{
+				ExecutionInfo.AddEvent(FAutomationEvent(EAutomationEventType::Error, FString::Printf(TEXT("The %s conversion to RawMesh is not lossless, polygon group material index array is different. PolygonGroupID [%d] expected Material index [%d] result [%d]"),
+					*AssetName,
+					PolygonGroupID.GetValue(),
+					ReferencePolygonGroupMaterialIndex[PolygonGroupID],
+					ResultPolygonGroupMaterialIndex[PolygonGroupID])));
+				bAllSame = false;
+				break;
+			}
+		}
+	}
 	return bAllSame;
 }
 
diff --git a/Engine/Source/Editor/UnrealEd/Private/Fbx/FbxStaticMeshImport.cpp b/Engine/Source/Editor/UnrealEd/Private/Fbx/FbxStaticMeshImport.cpp
--- a/Engine/Source/Editor/UnrealEd/Private/Fbx/FbxStaticMeshImport.cpp
+++ b/Engine/Source/Editor/UnrealEd/Private/Fbx/FbxStaticMeshImport.cpp
@@ -533,11 +533,8 @@ bool UnFbx::FFbxImporter::BuildStaticMeshFromGeometry(FbxNode* Node, UStaticMesh
 		TEdgeAttributeArray<bool>& EdgeHardnesses = MeshDescription->EdgeAttributes().GetAttributes<bool>(MeshAttribute::Edge::IsHard);
 		TEdgeAttributeArray<float>& EdgeCreaseSharpnesses = MeshDescription->EdgeAttributes().GetAttributes<float>(MeshAttribute::Edge::CreaseSharpness);
 
-		TPolygonGroupAttributeArray<FSoftObjectPath>& PolygonGroupMaterialAssets = MeshDescription->PolygonGroupAttributes().GetAttributes<FSoftObjectPath>(MeshAttribute::PolygonGroup::MaterialAsset);
-		TPolygonGroupAttributeArray<FName>& PolygonGroupMaterialSlotNames = MeshDescription->PolygonGroupAttributes().GetAttributes<FName>(MeshAttribute::PolygonGroup::MaterialSlotName);
 		TPolygonGroupAttributeArray<FName>& PolygonGroupImportedMaterialSlotNames = MeshDescription->PolygonGroupAttributes().GetAttributes<FName>(MeshAttribute::PolygonGroup::ImportedMaterialSlotName);
-		TPolygonGroupAttributeArray<bool>& PolygonGroupCollision = MeshDescription->PolygonGroupAttributes().GetAttributes<bool>(MeshAttribute::PolygonGroup::EnableCollision);
-		TPolygonGroupAttributeArray<bool>& PolygonGroupCastShadow = MeshDescription->PolygonGroupAttributes().GetAttributes<bool>(MeshAttribute::PolygonGroup::CastShadow);
+		TPolygonGroupAttributeArray<int>& PolygonGroupMaterialIndex = MeshDescription->PolygonGroupAttributes().GetAttributes<int>(MeshAttribute::PolygonGroup::MaterialIndex);
 
 		int32 VertexOffset = MeshDescription->Vertices().Num();
 		int32 VertexInstanceOffset = MeshDescription->VertexInstances().Num();
@@ -817,11 +814,8 @@ bool UnFbx::FFbxImporter::BuildStaticMeshFromGeometry(FbxNode* Node, UStaticMesh
 				if (ExistingPolygonGroup == FPolygonGroupID::Invalid)
 				{
 					ExistingPolygonGroup = MeshDescription->CreatePolygonGroup();
-					PolygonGroupCastShadow[ExistingPolygonGroup] = true;
-					PolygonGroupCollision[ExistingPolygonGroup] = bEnableCollision;
 					PolygonGroupImportedMaterialSlotNames[ExistingPolygonGroup] = ImportedMaterialSlotName;
-					PolygonGroupMaterialSlotNames[ExistingPolygonGroup] = ImportedMaterialSlotName;
-					PolygonGroupMaterialAssets[ExistingPolygonGroup] = FSoftObjectPath(Material);
+					PolygonGroupMaterialIndex[ExistingPolygonGroup] = RealMaterialIndex;
 				}
 				PolygonGroupMapping.Add(RealMaterialIndex, ExistingPolygonGroup);
 			}
@@ -1898,22 +1892,18 @@ UStaticMesh* UnFbx::FFbxImporter::ImportStaticMeshAsSingle(UObject* InParent, TA
 		}
 		else
 		{
-			TPolygonGroupAttributeArray<FSoftObjectPath>& PolygonGroupMaterialAssets = MeshDescription->PolygonGroupAttributes().GetAttributes<FSoftObjectPath>( MeshAttribute::PolygonGroup::MaterialAsset );
-			TPolygonGroupAttributeArray<FName>& PolygonGroupMaterialSlotNames = MeshDescription->PolygonGroupAttributes().GetAttributes<FName>( MeshAttribute::PolygonGroup::MaterialSlotName );
-			TPolygonGroupAttributeArray<FName>& PolygonGroupImportedMaterialSlotNames = MeshDescription->PolygonGroupAttributes().GetAttributes<FName>( MeshAttribute::PolygonGroup::ImportedMaterialSlotName );
-			TPolygonGroupAttributeArray<bool>& PolygonGroupCollision = MeshDescription->PolygonGroupAttributes().GetAttributes<bool>( MeshAttribute::PolygonGroup::EnableCollision );
-			TPolygonGroupAttributeArray<bool>& PolygonGroupCastShadow = MeshDescription->PolygonGroupAttributes().GetAttributes<bool>( MeshAttribute::PolygonGroup::CastShadow );
+			TPolygonGroupAttributeArray<FName>& PolygonGroupImportedMaterialSlotNames = MeshDescription->PolygonGroupAttributes().GetAttributes<FName>(MeshAttribute::PolygonGroup::ImportedMaterialSlotName);
+			TPolygonGroupAttributeArray<int>& PolygonGroupMaterialIndex = MeshDescription->PolygonGroupAttributes().GetAttributes<int>(MeshAttribute::PolygonGroup::MaterialIndex);
 
 			TArray<FStaticMaterial> MaterialToAdd;
 			for (const FPolygonGroupID PolygonGroupID : MeshDescription->PolygonGroups().GetElementIDs())
 			{
-				int32 MaterialIndex = PolygonGroupID.GetValue();
-				UMaterialInterface* Material = Cast<UMaterialInterface>(PolygonGroupMaterialAssets[PolygonGroupID].TryLoad());
-				const FName& MaterialSlotName = PolygonGroupMaterialSlotNames[PolygonGroupID];
+				int32 MaterialIndex = PolygonGroupMaterialIndex[PolygonGroupID];
+				UMaterialInterface* Material = MeshMaterials.IsValidIndex(MaterialIndex) ? MeshMaterials[MaterialIndex].Material : UMaterial::GetDefaultMaterial(MD_Surface);
 				const FName& ImportedMaterialSlotName = PolygonGroupImportedMaterialSlotNames[PolygonGroupID];
+				const FName MaterialSlotName = ImportedMaterialSlotName;
 				FStaticMaterial StaticMaterial(Material, MaterialSlotName, ImportedMaterialSlotName);
 				int32 AddedIndex = (LODIndex > 0) ? MaterialToAdd.Add(StaticMaterial) : StaticMesh->StaticMaterials.Add(StaticMaterial);
-				check(MaterialIndex == AddedIndex);
 			}
 			if (LODIndex > 0)
 			{
diff --git a/Engine/Source/Runtime/Engine/Classes/Engine/StaticMesh.h b/Engine/Source/Runtime/Engine/Classes/Engine/StaticMesh.h
--- a/Engine/Source/Runtime/Engine/Classes/Engine/StaticMesh.h
+++ b/Engine/Source/Runtime/Engine/Classes/Engine/StaticMesh.h
@@ -661,6 +661,12 @@ public:
 	ENGINE_API class UMeshDescription* GetOriginalMeshDescription(int32 LodIndex = 0);
 	ENGINE_API void SetOriginalMeshDescription(int32 LodIndex, class UMeshDescription* MeshDescription);
 	ENGINE_API void ClearOriginalMeshDescription(int32 LodIndex);
+
+	/**
+	 * Internal function use to make sure all imported material slot name are unique and non empty.
+	 */
+	void FixupMaterialSlotName();
+
 #endif
 
 	/** Pre-build navigation collision */
diff --git a/Engine/Source/Runtime/Engine/Private/StaticMesh.cpp b/Engine/Source/Runtime/Engine/Private/StaticMesh.cpp
--- a/Engine/Source/Runtime/Engine/Private/StaticMesh.cpp
+++ b/Engine/Source/Runtime/Engine/Private/StaticMesh.cpp
@@ -2379,6 +2379,46 @@ void UStaticMesh::ClearOriginalMeshDescription(int32 LodIndex)
 	}
 }
 
+void UStaticMesh::FixupMaterialSlotName()
+{
+	TArray<FName> UniqueMaterialSlotName;
+	//Make sure we have non empty imported material slot names
+	for (FStaticMaterial& Material : StaticMaterials)
+	{
+		if (Material.ImportedMaterialSlotName != NAME_None)
+		{
+			continue;
+		}
+		
+		if (Material.MaterialSlotName != NAME_None)
+		{
+			Material.ImportedMaterialSlotName = Material.MaterialSlotName;
+		}
+		else if (Material.MaterialInterface != nullptr)
+		{
+			Material.ImportedMaterialSlotName = Material.MaterialInterface->GetFName();
+		}
+		else
+		{
+			Material.ImportedMaterialSlotName = FName(TEXT("MaterialSlot"));
+		}
+
+		FString UniqueName = Material.ImportedMaterialSlotName.ToString();
+		int32 UniqueIndex = 1;
+		while (UniqueMaterialSlotName.Contains(FName(*UniqueName)))
+		{
+			UniqueName = FString::Printf(TEXT("%s_%d"), *UniqueName, UniqueIndex);
+			UniqueIndex++;
+		}
+		Material.ImportedMaterialSlotName = FName(*UniqueName);
+		UniqueMaterialSlotName.Add(Material.ImportedMaterialSlotName);
+		if (Material.MaterialSlotName == NAME_None)
+		{
+			Material.MaterialSlotName = Material.ImportedMaterialSlotName;
+		}
+	}
+}
+
 #endif
 
 void UStaticMesh::CacheDerivedData()
@@ -2772,6 +2812,8 @@ void UStaticMesh::PostLoad()
 			SetLODGroup(LODGroup);
 		}
 
+		FixupMaterialSlotName();
+
 		CacheDerivedData();
 
 		//Fix up the material to remove redundant material, this is needed since the material refactor where we do not have anymore copy of the materials
diff --git a/Engine/Source/Runtime/MeshBuilder/Private/MeshDescriptionHelper.cpp b/Engine/Source/Runtime/MeshBuilder/Private/MeshDescriptionHelper.cpp
--- a/Engine/Source/Runtime/MeshBuilder/Private/MeshDescriptionHelper.cpp
+++ b/Engine/Source/Runtime/MeshBuilder/Private/MeshDescriptionHelper.cpp
@@ -929,7 +929,6 @@ void FMeshDescriptionHelper::ConverToRawMesh(const UMeshDescription* SourceMeshD
 
 	//Gather all array data
 	const TVertexAttributeArray<FVector>& VertexPositions = SourceMeshDescription->VertexAttributes().GetAttributes<FVector>(MeshAttribute::Vertex::Position);
-	//const TVertexAttributeArray<float>& VertexCornerSharpness = SourceMeshDescription->VertexAttributes().GetAttributes<float>(MeshAttribute::Vertex::CornerSharpness);
 
 	const TVertexInstanceAttributeArray<FVector>& VertexInstanceNormals = SourceMeshDescription->VertexInstanceAttributes().GetAttributes<FVector>(MeshAttribute::VertexInstance::Normal);
 	const TVertexInstanceAttributeArray<FVector>& VertexInstanceTangents = SourceMeshDescription->VertexInstanceAttributes().GetAttributes<FVector>(MeshAttribute::VertexInstance::Tangent);
@@ -937,19 +936,7 @@ void FMeshDescriptionHelper::ConverToRawMesh(const UMeshDescription* SourceMeshD
 	const TVertexInstanceAttributeArray<FVector4>& VertexInstanceColors = SourceMeshDescription->VertexInstanceAttributes().GetAttributes<FVector4>(MeshAttribute::VertexInstance::Color);
 	const TVertexInstanceAttributeIndicesArray<FVector2D>& VertexInstanceUVs = SourceMeshDescription->VertexInstanceAttributes().GetAttributesSet<FVector2D>(MeshAttribute::VertexInstance::TextureCoordinate);
 
-	//const TEdgeAttributeArray<bool>& EdgeHardnesses = SourceMeshDescription->EdgeAttributes().GetAttributes<bool>(MeshAttribute::Edge::IsHard);
-	//const TEdgeAttributeArray<float>& EdgeCreaseSharpnesses = SourceMeshDescription->EdgeAttributes().GetAttributes<float>(MeshAttribute::Edge::CreaseSharpness);
-
-	//const TPolygonAttributeArray<FVector>& PolygonNormals = SourceMeshDescription->PolygonAttributes().GetAttributes<FVector>(MeshAttribute::Polygon::Normal);
-	//const TPolygonAttributeArray<FVector>& PolygonTangents = SourceMeshDescription->PolygonAttributes().GetAttributes<FVector>(MeshAttribute::Polygon::Tangent);
-	//const TPolygonAttributeArray<FVector>& PolygonBiNormals = SourceMeshDescription->PolygonAttributes().GetAttributes<FVector>(MeshAttribute::Polygon::Binormal);
-	//const TPolygonAttributeArray<FVector>& PolygonCenters = SourceMeshDescription->PolygonAttributes().GetAttributes<FVector>(MeshAttribute::Polygon::Center);
-
-	//const TPolygonGroupAttributeArray<FSoftObjectPath>& PolygonGroupMaterialAssets = SourceMeshDescription->PolygonGroupAttributes().GetAttributes<FSoftObjectPath>(MeshAttribute::PolygonGroup::MaterialAsset);
-	//const TPolygonGroupAttributeArray<FName>& PolygonGroupMaterialSlotNames = SourceMeshDescription->PolygonGroupAttributes().GetAttributes<FName>(MeshAttribute::PolygonGroup::MaterialSlotName);
-	//const TPolygonGroupAttributeArray<FName>& PolygonGroupImportedMaterialSlotNames = SourceMeshDescription->PolygonGroupAttributes().GetAttributes<FName>(MeshAttribute::PolygonGroup::ImportedMaterialSlotName);
-	//const TPolygonGroupAttributeArray<bool>& PolygonGroupCollision = SourceMeshDescription->PolygonGroupAttributes().GetAttributes<bool>(MeshAttribute::PolygonGroup::EnableCollision);
-	//const TPolygonGroupAttributeArray<bool>& PolygonGroupCastShadow = SourceMeshDescription->PolygonGroupAttributes().GetAttributes<bool>(MeshAttribute::PolygonGroup::CastShadow);
+	const TPolygonGroupAttributeArray<int>& PolygonGroupMaterialIndex = SourceMeshDescription->PolygonGroupAttributes().GetAttributes<int>(MeshAttribute::PolygonGroup::MaterialIndex);
 
 	DestinationRawMesh.VertexPositions.AddZeroed(SourceMeshDescription->Vertices().Num());
 	for (const FVertexID& VertexID : SourceMeshDescription->Vertices().GetElementIDs())
@@ -984,7 +971,7 @@ void FMeshDescriptionHelper::ConverToRawMesh(const UMeshDescription* SourceMeshD
 		const TArray<FMeshTriangle>& Triangles = SourceMeshDescription->GetPolygonTriangles(PolygonID);
 		for (const FMeshTriangle& MeshTriangle : Triangles)
 		{
-			DestinationRawMesh.FaceMaterialIndices[PolygonIDValue] = PolygonGroupID.GetValue();
+			DestinationRawMesh.FaceMaterialIndices[PolygonIDValue] = PolygonGroupMaterialIndex[PolygonGroupID];
 			DestinationRawMesh.FaceSmoothingMasks[PolygonIDValue] = 0; //Conversion of soft/hard to smooth mask is done after the geometry is converted
 			for (int32 Corner = 0; Corner < 3; ++Corner)
 			{
@@ -1011,7 +998,6 @@ void FMeshDescriptionHelper::ConverFromRawMesh(const struct FRawMesh &SourceRawM
 	DestinationMeshDescription->Empty();
 	//Gather all array data
 	TVertexAttributeArray<FVector>& VertexPositions = DestinationMeshDescription->VertexAttributes().GetAttributes<FVector>(MeshAttribute::Vertex::Position);
-	//TVertexAttributeArray<float>& VertexCornerSharpness = DestinationMeshDescription->VertexAttributes().GetAttributes<float>(MeshAttribute::Vertex::CornerSharpness);
 
 	TVertexInstanceAttributeArray<FVector>& VertexInstanceNormals = DestinationMeshDescription->VertexInstanceAttributes().GetAttributes<FVector>(MeshAttribute::VertexInstance::Normal);
 	TVertexInstanceAttributeArray<FVector>& VertexInstanceTangents = DestinationMeshDescription->VertexInstanceAttributes().GetAttributes<FVector>(MeshAttribute::VertexInstance::Tangent);
@@ -1019,19 +1005,8 @@ void FMeshDescriptionHelper::ConverFromRawMesh(const struct FRawMesh &SourceRawM
 	TVertexInstanceAttributeArray<FVector4>& VertexInstanceColors = DestinationMeshDescription->VertexInstanceAttributes().GetAttributes<FVector4>(MeshAttribute::VertexInstance::Color);
 	TVertexInstanceAttributeIndicesArray<FVector2D>& VertexInstanceUVs = DestinationMeshDescription->VertexInstanceAttributes().GetAttributesSet<FVector2D>(MeshAttribute::VertexInstance::TextureCoordinate);
 
-	//TEdgeAttributeArray<bool>& EdgeHardnesses = DestinationMeshDescription->EdgeAttributes().GetAttributes<bool>(MeshAttribute::Edge::IsHard);
-	//TEdgeAttributeArray<float>& EdgeCreaseSharpnesses = DestinationMeshDescription->EdgeAttributes().GetAttributes<float>(MeshAttribute::Edge::CreaseSharpness);
-
-	//TPolygonAttributeArray<FVector>& PolygonNormals = DestinationMeshDescription->PolygonAttributes().GetAttributes<FVector>(MeshAttribute::Polygon::Normal);
-	//TPolygonAttributeArray<FVector>& PolygonTangents = DestinationMeshDescription->PolygonAttributes().GetAttributes<FVector>(MeshAttribute::Polygon::Tangent);
-	//TPolygonAttributeArray<FVector>& PolygonBiNormals = DestinationMeshDescription->PolygonAttributes().GetAttributes<FVector>(MeshAttribute::Polygon::Binormal);
-	//TPolygonAttributeArray<FVector>& PolygonCenters = DestinationMeshDescription->PolygonAttributes().GetAttributes<FVector>(MeshAttribute::Polygon::Center);
-
-	TPolygonGroupAttributeArray<FSoftObjectPath>& PolygonGroupMaterialAssets = DestinationMeshDescription->PolygonGroupAttributes().GetAttributes<FSoftObjectPath>(MeshAttribute::PolygonGroup::MaterialAsset);
-	TPolygonGroupAttributeArray<FName>& PolygonGroupMaterialSlotNames = DestinationMeshDescription->PolygonGroupAttributes().GetAttributes<FName>(MeshAttribute::PolygonGroup::MaterialSlotName);
 	TPolygonGroupAttributeArray<FName>& PolygonGroupImportedMaterialSlotNames = DestinationMeshDescription->PolygonGroupAttributes().GetAttributes<FName>(MeshAttribute::PolygonGroup::ImportedMaterialSlotName);
-	TPolygonGroupAttributeArray<bool>& PolygonGroupCollision = DestinationMeshDescription->PolygonGroupAttributes().GetAttributes<bool>(MeshAttribute::PolygonGroup::EnableCollision);
-	TPolygonGroupAttributeArray<bool>& PolygonGroupCastShadow = DestinationMeshDescription->PolygonGroupAttributes().GetAttributes<bool>(MeshAttribute::PolygonGroup::CastShadow);
+	TPolygonGroupAttributeArray<int>& PolygonGroupMaterialIndex = DestinationMeshDescription->PolygonGroupAttributes().GetAttributes<int>(MeshAttribute::PolygonGroup::MaterialIndex);
 
 	int32 NumTexCoords = 0;
 	int32 MaxTexCoords = FMath::Min<int32>(MAX_MESH_TEXTURE_COORDS, MAX_STATIC_TEXCOORDS);
@@ -1066,21 +1041,22 @@ void FMeshDescriptionHelper::ConverFromRawMesh(const struct FRawMesh &SourceRawM
 		int32 VerticeIndexBase = TriangleIndex * 3;
 
 		//PolygonGroup
-		FPolygonGroupID PolygonGroupID(SourceRawMesh.FaceMaterialIndices[TriangleIndex]);
-		if (PolygonGroupID.GetValue() >= PolygonGroups.Num())
+		FPolygonGroupID PolygonGroupID = FPolygonGroupID::Invalid;
+		int32 MaterialIndex = SourceRawMesh.FaceMaterialIndices[TriangleIndex];
+		for (const FPolygonGroupID& SearchPolygonGroupID : DestinationMeshDescription->PolygonGroups().GetElementIDs())
 		{
-			for (int32 AddMaterialIndex = PolygonGroups.Num(); AddMaterialIndex <= PolygonGroupID.GetValue(); ++AddMaterialIndex)
+			if (PolygonGroupMaterialIndex[SearchPolygonGroupID] == MaterialIndex)
 			{
-				FPolygonGroupID AddedPolygonGroupID = DestinationMeshDescription->CreatePolygonGroup();
-				check(AddMaterialIndex == AddedPolygonGroupID.GetValue());
-				PolygonGroupCastShadow[AddedPolygonGroupID] = true;
-				PolygonGroupCollision[AddedPolygonGroupID] = true;
-				PolygonGroupImportedMaterialSlotNames[AddedPolygonGroupID] = NAME_None;
-				PolygonGroupMaterialSlotNames[AddedPolygonGroupID] = NAME_None;
-				PolygonGroupMaterialAssets[AddedPolygonGroupID] = FSoftObjectPath(UMaterial::GetDefaultMaterial(MD_Surface));
-				PolygonGroups.Add(AddedPolygonGroupID);
+				PolygonGroupID = SearchPolygonGroupID;
 			}
 		}
+		if (PolygonGroupID == FPolygonGroupID::Invalid)
+		{
+			PolygonGroupID = DestinationMeshDescription->CreatePolygonGroup();
+			PolygonGroupMaterialIndex[PolygonGroupID] = MaterialIndex;
+			PolygonGroupImportedMaterialSlotNames[PolygonGroupID] = FName(*FString::Printf(TEXT("MaterialSlot_%d"), MaterialIndex));
+			PolygonGroups.Add(PolygonGroupID);
+		}
 
 		for (int32 Corner = 0; Corner < 3; ++Corner)
 		{
diff --git a/Engine/Source/Runtime/MeshBuilder/Private/StaticMeshBuilder.cpp b/Engine/Source/Runtime/MeshBuilder/Private/StaticMeshBuilder.cpp
--- a/Engine/Source/Runtime/MeshBuilder/Private/StaticMeshBuilder.cpp
+++ b/Engine/Source/Runtime/MeshBuilder/Private/StaticMeshBuilder.cpp
@@ -388,9 +388,6 @@ void BuildVertexBuffer(
 	const uint32 NumTextureCoord = MeshDescription->VertexInstanceAttributes().GetAttributeIndexCount<FVector2D>(MeshAttribute::VertexInstance::TextureCoordinate);
 
 	const TPolygonGroupAttributeArray<FName>& PolygonGroupImportedMaterialSlotNames = MeshDescription->PolygonGroupAttributes().GetAttributes<FName>(MeshAttribute::PolygonGroup::ImportedMaterialSlotName);
-	const TPolygonGroupAttributeArray<FSoftObjectPath>& PolygonGroupMaterialAssets = MeshDescription->PolygonGroupAttributes().GetAttributes<FSoftObjectPath>(MeshAttribute::PolygonGroup::MaterialAsset);
-	const TPolygonGroupAttributeArray<bool>& PolygonGroupCollision = MeshDescription->PolygonGroupAttributes().GetAttributes<bool>(MeshAttribute::PolygonGroup::EnableCollision);
-	const TPolygonGroupAttributeArray<bool>& PolygonGroupCastShadow = MeshDescription->PolygonGroupAttributes().GetAttributes<bool>(MeshAttribute::PolygonGroup::CastShadow);
 
 	const TVertexAttributeArray<FVector>& VertexPositions = MeshDescription->VertexAttributes().GetAttributes<FVector>( MeshAttribute::Vertex::Position );
 	const TVertexInstanceAttributeArray<FVector>& VertexInstanceNormals = MeshDescription->VertexInstanceAttributes().GetAttributes<FVector>( MeshAttribute::VertexInstance::Normal );
diff --git a/Engine/Source/Runtime/MeshDescription/Private/MeshAttributes.cpp b/Engine/Source/Runtime/MeshDescription/Private/MeshAttributes.cpp
--- a/Engine/Source/Runtime/MeshDescription/Private/MeshAttributes.cpp
+++ b/Engine/Source/Runtime/MeshDescription/Private/MeshAttributes.cpp
@@ -26,5 +26,6 @@ namespace MeshAttribute
 	const FName PolygonGroup::ImportedMaterialSlotName( "ImportedMaterialSlotName" );
 	const FName PolygonGroup::EnableCollision( "EnableCollision" );
 	const FName PolygonGroup::CastShadow( "CastShadow" );
+	const FName PolygonGroup::MaterialIndex( "MaterialIndex" );
 }
 
diff --git a/Engine/Source/Runtime/MeshDescription/Private/MeshDescription.cpp b/Engine/Source/Runtime/MeshDescription/Private/MeshDescription.cpp
--- a/Engine/Source/Runtime/MeshDescription/Private/MeshDescription.cpp
+++ b/Engine/Source/Runtime/MeshDescription/Private/MeshDescription.cpp
@@ -35,11 +35,12 @@ void UMeshDescription::RegisterBaseAttributes()
 	PolygonAttributes().RegisterAttribute<FVector>(MeshAttribute::Polygon::Center, 1, FVector::ZeroVector, EMeshAttributeFlags::AutoGenerated);
 
 	// Add basic polygon group attributes
-	PolygonGroupAttributes().RegisterAttribute<FSoftObjectPath>(MeshAttribute::PolygonGroup::MaterialAsset);
-	PolygonGroupAttributes().RegisterAttribute<FName>(MeshAttribute::PolygonGroup::MaterialSlotName);
-	PolygonGroupAttributes().RegisterAttribute<FName>(MeshAttribute::PolygonGroup::ImportedMaterialSlotName);
-	PolygonGroupAttributes().RegisterAttribute<bool>(MeshAttribute::PolygonGroup::EnableCollision);
-	PolygonGroupAttributes().RegisterAttribute<bool>(MeshAttribute::PolygonGroup::CastShadow);
+	PolygonGroupAttributes().RegisterAttribute<FSoftObjectPath>(MeshAttribute::PolygonGroup::MaterialAsset); //Deprecated
+	PolygonGroupAttributes().RegisterAttribute<FName>(MeshAttribute::PolygonGroup::MaterialSlotName); //Deprecated
+	PolygonGroupAttributes().RegisterAttribute<FName>(MeshAttribute::PolygonGroup::ImportedMaterialSlotName); //The unique key to match the mesh material slot
+	PolygonGroupAttributes().RegisterAttribute<bool>(MeshAttribute::PolygonGroup::EnableCollision); //Deprecated
+	PolygonGroupAttributes().RegisterAttribute<bool>(MeshAttribute::PolygonGroup::CastShadow); //Deprecated
+	PolygonGroupAttributes().RegisterAttribute<int>(MeshAttribute::PolygonGroup::MaterialIndex); //This attribute is used for FRawMesh conversion only and is temporary, use the ImportedMaterialSlotName to match the mesh material slot
 }
 
 void UMeshDescription::UnRegisterBaseAttributes()
@@ -71,6 +72,7 @@ void UMeshDescription::UnRegisterBaseAttributes()
 	PolygonGroupAttributes().UnregisterAttribute<FName>(MeshAttribute::PolygonGroup::ImportedMaterialSlotName);
 	PolygonGroupAttributes().UnregisterAttribute<bool>(MeshAttribute::PolygonGroup::EnableCollision);
 	PolygonGroupAttributes().UnregisterAttribute<bool>(MeshAttribute::PolygonGroup::CastShadow);
+	PolygonGroupAttributes().UnregisterAttribute<int>(MeshAttribute::PolygonGroup::MaterialIndex);
 }
 
 void UMeshDescription::Serialize( FArchive& Ar )
diff --git a/Engine/Source/Runtime/MeshDescription/Public/MeshAttributes.h b/Engine/Source/Runtime/MeshDescription/Public/MeshAttributes.h
--- a/Engine/Source/Runtime/MeshDescription/Public/MeshAttributes.h
+++ b/Engine/Source/Runtime/MeshDescription/Public/MeshAttributes.h
@@ -45,5 +45,6 @@ namespace MeshAttribute
 		extern MESHDESCRIPTION_API const FName ImportedMaterialSlotName;
 		extern MESHDESCRIPTION_API const FName EnableCollision;
 		extern MESHDESCRIPTION_API const FName CastShadow;
+		extern MESHDESCRIPTION_API const FName MaterialIndex;
 	}
 }
