Add reverse polygon vertex order
Add Generate UVs API to MeshBuilder

diff --git a/Engine/Source/Runtime/MeshBuilder/Private/MeshBuilder.cpp b/Engine/Source/Runtime/MeshBuilder/Private/MeshBuilder.cpp
--- a/Engine/Source/Runtime/MeshBuilder/Private/MeshBuilder.cpp
+++ b/Engine/Source/Runtime/MeshBuilder/Private/MeshBuilder.cpp
@@ -1,10 +1,14 @@
 // Copyright 1998-2017 Epic Games, Inc. All Rights Reserved.
-
 #include "MeshBuilder.h"
+#include "CoreMinimal.h"
+#include "CoreTypes.h"
+#include "UObject/UObjectGlobals.h"
+#include "UObject/Package.h"
 #include "MeshDescriptionHelper.h"
 #include "MeshDescription.h"
 #include "MeshAttributes.h"
 #include "RawMesh.h"
+#include "LayoutUV.h"
 
 
 FMeshBuilder::FMeshBuilder()
@@ -81,3 +85,38 @@ void FMeshDescriptionOperations::ConverFromRawMesh(const FRawMesh &SourceRawMesh
 {
 	FMeshDescriptionHelper::ConverFromRawMesh(SourceRawMesh, DestinationMeshDescription);
 }
+
+bool FMeshDescriptionOperations::GenerateUniqueUVsForStaticMesh(const UMeshDescription* MeshDescription, int32 TextureResolution, TArray<FVector2D>& OutTexCoords)
+{
+	// Create a copy of original mesh (only copy necessary data)
+	UMeshDescription* DuplicateMeshDescription = Cast<UMeshDescription>(StaticDuplicateObject(MeshDescription, GetTransientPackage(), NAME_None, RF_NoFlags));
+
+	// Find overlapping corners for UV generator. Allow some threshold - this should not produce any error in a case if resulting
+	// mesh will not merge these vertices.
+	TMultiMap<int32, int32> OverlappingCorners;
+	FMeshDescriptionHelper::FindOverlappingCorners(OverlappingCorners, DuplicateMeshDescription, THRESH_POINTS_ARE_SAME);
+
+	// Generate new UVs
+	FLayoutUV Packer(DuplicateMeshDescription, 0, 1, FMath::Clamp(TextureResolution / 4, 32, 512));
+	Packer.FindCharts(OverlappingCorners);
+
+	bool bPackSuccess = Packer.FindBestPacking();
+	if (bPackSuccess)
+	{
+		Packer.CommitPackedUVs();
+		TVertexInstanceAttributeIndicesArray<FVector2D>& VertexInstanceUVs = DuplicateMeshDescription->VertexInstanceAttributes().GetAttributesSet<FVector2D>(MeshAttribute::VertexInstance::TextureCoordinate);
+		// Save generated UVs
+		check(VertexInstanceUVs.GetNumIndices() > 1);
+		auto& UniqueUVsArray = VertexInstanceUVs.GetArrayForIndex(1);
+		OutTexCoords.AddZeroed(UniqueUVsArray.Num());
+		int32 TextureCoordIndex = 0;
+		for(const FVertexInstanceID& VertexInstanceID : DuplicateMeshDescription->VertexInstances().GetElementIDs())
+		{
+			OutTexCoords[TextureCoordIndex++] = UniqueUVsArray[VertexInstanceID];
+		}
+	}
+	//Make sure the transient duplicate will be GC
+	DuplicateMeshDescription->MarkPendingKill();
+
+	return bPackSuccess;
+}
\ No newline at end of file
diff --git a/Engine/Source/Runtime/MeshBuilder/Public/MeshBuilder.h b/Engine/Source/Runtime/MeshBuilder/Public/MeshBuilder.h
--- a/Engine/Source/Runtime/MeshBuilder/Public/MeshBuilder.h
+++ b/Engine/Source/Runtime/MeshBuilder/Public/MeshBuilder.h
@@ -26,13 +26,16 @@ private:
 class MESHBUILDER_API FMeshDescriptionOperations
 {
 public:
-	/** Compute the normals tangents and binormals for the mesh description */
+	/** Compute the normals tangents and binormals for the mesh description. */
 	static void ComputeMeshNTBs(class UMeshDescription* MeshDescription, const struct FMeshBuildSettings& BuildSettings);
 
-	/** Convert this mesh description into the old FRawMesh format*/
+	/** Convert this mesh description into the old FRawMesh format. */
 	static void ConverToRawMesh(const class UMeshDescription* SourceMeshDescription, struct FRawMesh &DestinationRawMesh);
 
-	/** Convert old FRawMesh format to MeshDescription*/
+	/** Convert old FRawMesh format to MeshDescription. */
 	static void ConverFromRawMesh(const struct FRawMesh &SourceRawMesh, class UMeshDescription* DestinationMeshDescription);
+
+	/** Create some UVs from the specified mesh description data. */
+	static bool GenerateUniqueUVsForStaticMesh(const class UMeshDescription* MeshDescription, int32 TextureResolution, TArray<FVector2D>& OutTexCoords);
 };
 
diff --git a/Engine/Source/Runtime/MeshDescription/Private/MeshDescription.cpp b/Engine/Source/Runtime/MeshDescription/Private/MeshDescription.cpp
--- a/Engine/Source/Runtime/MeshDescription/Private/MeshDescription.cpp
+++ b/Engine/Source/Runtime/MeshDescription/Private/MeshDescription.cpp
@@ -506,4 +506,36 @@ void UMeshDescription::TriangulateMesh()
 		FMeshPolygon& Polygon = PolygonArray[PolygonID];
 		ComputePolygonTriangulation(PolygonID, Polygon.Triangles);
 	}
+}
+
+void UMeshDescription::ReversePolygonFacing(const FPolygonID PolygonID)
+{
+	FMeshPolygon& Polygon = GetPolygon(PolygonID);
+
+	//Build a reverse perimeter
+	TArray<FVertexInstanceID> ReverseVertexInstanceIDs;
+	int32 PerimeterReverseIndex = Polygon.PerimeterContour.VertexInstanceIDs.Num() - 1;
+	for(const FVertexInstanceID& VertexInstanceID : Polygon.PerimeterContour.VertexInstanceIDs)
+	{
+		ReverseVertexInstanceIDs[PerimeterReverseIndex] = VertexInstanceID;
+		PerimeterReverseIndex--;
+	}
+	//Assign the reverse perimeter
+	for (int32 VertexInstanceIndex = 0; VertexInstanceIndex < ReverseVertexInstanceIDs.Num(); ++VertexInstanceIndex)
+	{
+		Polygon.PerimeterContour.VertexInstanceIDs[VertexInstanceIndex] = ReverseVertexInstanceIDs[VertexInstanceIndex];
+	}
+	
+	//Triangulate the polygon since we reverse the indices
+	ComputePolygonTriangulation(PolygonID, Polygon.Triangles);
+}
+
+void UMeshDescription::ReverseAllPolygonFacing()
+{
+	// Perform triangulation directly into mesh polygons
+	for (const FPolygonID PolygonID : Polygons().GetElementIDs())
+	{
+		FMeshPolygon& Polygon = PolygonArray[PolygonID];
+		ReversePolygonFacing(PolygonID);
+	}
 }
\ No newline at end of file
diff --git a/Engine/Source/Runtime/MeshDescription/Public/MeshDescription.h b/Engine/Source/Runtime/MeshDescription/Public/MeshDescription.h
--- a/Engine/Source/Runtime/MeshDescription/Public/MeshDescription.h
+++ b/Engine/Source/Runtime/MeshDescription/Public/MeshDescription.h
@@ -607,7 +607,9 @@ public:
 		return PolygonGroupArray.IsValid( PolygonGroupID );
 	}
 
-
+//////////////////////////////////////////////////////////////////////////
+// Meshdescription general functions
+public:
 	/** Returns the edge ID defined by the two given vertex IDs, if there is one; otherwise FEdgeID::Invalid */
 	FEdgeID GetVertexPairEdge( const FVertexID VertexID0, const FVertexID VertexID1 ) const
 	{
@@ -740,10 +742,12 @@ public:
 	FString GetIdString();
 #endif
 
-public:
 	void ComputePolygonTriangulation(const FPolygonID PolygonID, TArray<FMeshTriangle>& OutTriangles);
 	void TriangulateMesh();
 
+	void ReversePolygonFacing(const FPolygonID PolygonID);
+	void ReverseAllPolygonFacing();
+
 private:
 	bool VectorsOnSameSide(const FVector& Vec, const FVector& A, const FVector& B, const float SameSideDotProductEpsilon);
 	bool PointInTriangle(const FVector& A, const FVector& B, const FVector& C, const FVector& P, const float InsideTriangleDotProductEpsilon);
