Fixed macro redefinition clash in Unity builds.

diff --git a/Engine/Source/Runtime/MeshBuilder/Private/LayoutUV.cpp b/Engine/Source/Runtime/MeshBuilder/Private/LayoutUV.cpp
--- a/Engine/Source/Runtime/MeshBuilder/Private/LayoutUV.cpp
+++ b/Engine/Source/Runtime/MeshBuilder/Private/LayoutUV.cpp
@@ -24,6 +24,7 @@ FLayoutUV::FLayoutUV( UMeshDescription* InMesh, uint32 InSrcChannel, uint32 InDs
 
 void FLayoutUV::FindCharts( const TMultiMap<int32,int32>& OverlappingCorners )
 {
+	const float ThreshUVsAreSame = GetUVEqualityThreshold();
 	double Begin = FPlatformTime::Seconds();
 
 	uint32 NumIndexes = MeshDescription->VertexInstances().Num();
@@ -88,7 +89,7 @@ void FLayoutUV::FindCharts( const TMultiMap<int32,int32>& OverlappingCorners )
 							FVector2D EdgeUVj = TexCoords[je] - TexCoords[j];
 							
 							// Would these edges match if the charts were translated
-							bool bTranslatedUVMatch = ( EdgeUVi - EdgeUVj ).IsNearlyZero( THRESH_UVS_ARE_SAME );
+							bool bTranslatedUVMatch = ( EdgeUVi - EdgeUVj ).IsNearlyZero( ThreshUVsAreSame );
 							if( bTranslatedUVMatch )
 							{
 								// Note: may be mirrored
@@ -282,7 +283,7 @@ void FLayoutUV::FindCharts( const TMultiMap<int32,int32>& OverlappingCorners )
 					
 					FVector2D EdgeOffset0 = UV0i - UV1j;
 					FVector2D EdgeOffset1 = UV1i - UV0j;
-					checkSlow( ( EdgeOffset0 - EdgeOffset1 ).IsNearlyZero( THRESH_UVS_ARE_SAME ) );
+					checkSlow( ( EdgeOffset0 - EdgeOffset1 ).IsNearlyZero( ThreshUVsAreSame ) );
 
 					FVector2D Translation = EdgeOffset0;
 
@@ -307,9 +308,9 @@ void FLayoutUV::FindCharts( const TMultiMap<int32,int32>& OverlappingCorners )
 						uint32 Sign = Side & 1;
 						uint32 Axis = Side >> 1;
 
-						bool bAxisAligned = FMath::Abs( EdgeUVi[ Axis ] ) < THRESH_UVS_ARE_SAME;
-						bool bBorderA = FMath::Abs( UV0i[ Axis ] - ( Sign ^ 0 ? Chart.MaxUV[ Axis ] : Chart.MinUV[ Axis ] ) ) < THRESH_UVS_ARE_SAME;
-						bool bBorderB = FMath::Abs( UV0j[ Axis ] - ( Sign ^ 1 ? Chart.MaxUV[ Axis ] : Chart.MinUV[ Axis ] ) ) < THRESH_UVS_ARE_SAME;
+						bool bAxisAligned = FMath::Abs( EdgeUVi[ Axis ] ) < ThreshUVsAreSame;
+						bool bBorderA = FMath::Abs( UV0i[ Axis ] - ( Sign ^ 0 ? Chart.MaxUV[ Axis ] : Chart.MinUV[ Axis ] ) ) < ThreshUVsAreSame;
+						bool bBorderB = FMath::Abs( UV0j[ Axis ] - ( Sign ^ 1 ? Chart.MaxUV[ Axis ] : Chart.MinUV[ Axis ] ) ) < ThreshUVsAreSame;
 
 						// FIXME mirrored
 						if( !bAxisAligned || !bBorderA || !bBorderB )
@@ -331,9 +332,9 @@ void FLayoutUV::FindCharts( const TMultiMap<int32,int32>& OverlappingCorners )
 						FVector2D ExtentDiff = ExtentA - ExtentB;
 						FVector2D Separation = ExtentA + ExtentB + CenterDiff * ( Sign ? 1.0f : -1.0f );
 
-						bool bCenterMatch = FMath::Abs( CenterDiff[ Axis ^ 1 ] ) < THRESH_UVS_ARE_SAME;
-						bool bExtentMatch = FMath::Abs( ExtentDiff[ Axis ^ 1 ] ) < THRESH_UVS_ARE_SAME;
-						bool bSeparate    = FMath::Abs( Separation[ Axis ^ 0 ] ) < THRESH_UVS_ARE_SAME;
+						bool bCenterMatch = FMath::Abs( CenterDiff[ Axis ^ 1 ] ) < ThreshUVsAreSame;
+						bool bExtentMatch = FMath::Abs( ExtentDiff[ Axis ^ 1 ] ) < ThreshUVsAreSame;
+						bool bSeparate    = FMath::Abs( Separation[ Axis ^ 0 ] ) < ThreshUVsAreSame;
 	
 						if( !bCenterMatch || !bExtentMatch || !bSeparate )
 						{
diff --git a/Engine/Source/Runtime/MeshBuilder/Private/LayoutUV.h b/Engine/Source/Runtime/MeshBuilder/Private/LayoutUV.h
--- a/Engine/Source/Runtime/MeshBuilder/Private/LayoutUV.h
+++ b/Engine/Source/Runtime/MeshBuilder/Private/LayoutUV.h
@@ -90,9 +90,6 @@ private:
 };
 
 
-#define THRESH_UVS_ARE_SAME (GetUVEqualityThreshold())
-
-
 inline bool FLayoutUV::PositionsMatch( uint32 a, uint32 b ) const
 {
 	const FVertexInstanceID VertexInstanceIDA(a);
@@ -136,7 +133,7 @@ inline bool FLayoutUV::UVsMatch( uint32 a, uint32 b ) const
 	const FVertexInstanceID VertexInstanceIDB(b);
 
 	const TVertexInstanceAttributeArray<FVector2D>& VertexUVs = MeshDescription->VertexInstanceAttributes().GetAttributes<FVector2D>(MeshAttribute::VertexInstance::TextureCoordinate, SrcChannel);
-	return VertexUVs[VertexInstanceIDA].Equals(VertexUVs[VertexInstanceIDB], THRESH_UVS_ARE_SAME);
+	return VertexUVs[VertexInstanceIDA].Equals(VertexUVs[VertexInstanceIDB], GetUVEqualityThreshold());
 }
 
 inline bool FLayoutUV::VertsMatch( uint32 a, uint32 b ) const
