Linux build fix.

diff --git a/Engine/Source/Runtime/MeshBuilder/Private/MeshDescriptionHelper.cpp b/Engine/Source/Runtime/MeshBuilder/Private/MeshDescriptionHelper.cpp
--- a/Engine/Source/Runtime/MeshBuilder/Private/MeshDescriptionHelper.cpp
+++ b/Engine/Source/Runtime/MeshBuilder/Private/MeshDescriptionHelper.cpp
@@ -1,5 +1,4 @@
 // Copyright 1998-2017 Epic Games, Inc. All Rights Reserved.
-#pragma once
 
 #include "MeshDescriptionHelper.h"
 
