Fix circular dependencies, I move all the meshdescription operations code outside of the meshbuilder. Engine depend on MeshDescriptionOperations(MDO), MeshBuilder depend on MDO.
MDO do not depend on Engine and MeshBuilder.

diff --git a/Engine/Plugins/Tests/EditorTests/Source/EditorTests/EditorTests.Build.cs b/Engine/Plugins/Tests/EditorTests/Source/EditorTests/EditorTests.Build.cs
--- a/Engine/Plugins/Tests/EditorTests/Source/EditorTests/EditorTests.Build.cs
+++ b/Engine/Plugins/Tests/EditorTests/Source/EditorTests/EditorTests.Build.cs
@@ -65,6 +65,7 @@ public class EditorTests : ModuleRules
 				"MeshMergeUtilities",
 				"MaterialBaking",
                 "MeshDescription",
+                "MeshDescriptionOperations",
                 "MeshBuilder",
                 "RawMesh",
 			}
diff --git a/Engine/Plugins/Tests/EditorTests/Source/EditorTests/Private/UnrealEd/MeshDescriptionTests.cpp b/Engine/Plugins/Tests/EditorTests/Source/EditorTests/Private/UnrealEd/MeshDescriptionTests.cpp
--- a/Engine/Plugins/Tests/EditorTests/Source/EditorTests/Private/UnrealEd/MeshDescriptionTests.cpp
+++ b/Engine/Plugins/Tests/EditorTests/Source/EditorTests/Private/UnrealEd/MeshDescriptionTests.cpp
@@ -9,6 +9,7 @@
 #include "Engine/StaticMesh.h"
 #include "Materials/Material.h"
 #include "MeshDescription.h"
+#include "MeshDescriptionOperations.h"
 #include "MeshBuilder.h"
 #include "MeshAttributes.h"
 #include "RawMesh.h"
diff --git a/Engine/Source/Developer/MeshDescriptionOperations/MeshDescriptionOperations.Build.cs b/Engine/Source/Developer/MeshDescriptionOperations/MeshDescriptionOperations.Build.cs
new file mode 100644
--- /dev/null
+++ b/Engine/Source/Developer/MeshDescriptionOperations/MeshDescriptionOperations.Build.cs
@@ -0,0 +1,27 @@
+// Copyright 1998-2017 Epic Games, Inc. All Rights Reserved.
+
+namespace UnrealBuildTool.Rules
+{
+    public class MeshDescriptionOperations : ModuleRules
+    {
+        public MeshDescriptionOperations(ReadOnlyTargetRules Target) : base(Target)
+        {
+            PrivateIncludePaths.Add("Developer/MeshDescriptionOperations/Private");
+            PublicIncludePaths.Add("Developer/MeshDescriptionOperations/Public");
+
+            PrivateDependencyModuleNames.AddRange(
+                new string[]
+                {
+                    "Core",
+                    "CoreUObject",
+                    "RenderCore",
+                    "MeshDescription",
+                    "RawMesh",
+                    "RHI"
+                }
+            );
+
+            AddEngineThirdPartyPrivateStaticDependencies(Target, "MikkTSpace");
+        }
+    }
+}
diff --git a/Engine/Source/Developer/MeshDescriptionOperations/Private/MeshDescriptionOperations.cpp b/Engine/Source/Developer/MeshDescriptionOperations/Private/MeshDescriptionOperations.cpp
new file mode 100644
--- /dev/null
+++ b/Engine/Source/Developer/MeshDescriptionOperations/Private/MeshDescriptionOperations.cpp
@@ -0,0 +1,829 @@
+// Copyright 1998-2018 Epic Games, Inc. All Rights Reserved.
+
+#include "MeshDescriptionOperations.h"
+#include "MeshDescription.h"
+#include "MeshAttributes.h"
+#include "RawMesh.h"
+#include "RenderUtils.h"
+#include "mikktspace.h"
+
+DEFINE_LOG_CATEGORY(LogMeshDescriptionOperations);
+
+#define LOCTEXT_NAMESPACE "MeshDescriptionOperations"
+
+//////////////////////////////////////////////////////////////////////////
+// Local structure
+struct FVertexInfo
+{
+	FVertexInfo()
+	{
+		PolygonID = FPolygonID::Invalid;
+		VertexInstanceID = FVertexInstanceID::Invalid;
+		UVs = FVector2D(0.0f, 0.0f);
+		EdgeIDs.Reserve(2);//Most of the time a edge has two triangles
+	}
+
+	FPolygonID PolygonID;
+	FVertexInstanceID VertexInstanceID;
+	FVector2D UVs;
+	TArray<FEdgeID> EdgeIDs;
+};
+
+//////////////////////////////////////////////////////////////////////////
+// Converters
+
+void FMeshDescriptionOperations::ConvertHardEdgesToSmoothGroup(const UMeshDescription* SourceMeshDescription, struct FRawMesh &DestinationRawMesh)
+{
+	TMap<FPolygonID, uint32> PolygonSmoothGroup;
+	PolygonSmoothGroup.Reserve(SourceMeshDescription->Polygons().Num());
+	TArray<bool> ConsumedPolygons;
+	ConsumedPolygons.AddZeroed(SourceMeshDescription->Polygons().Num());
+
+	TMap < FPolygonID, uint32> PolygonAvoidances;
+
+	const TEdgeAttributeArray<bool>& EdgeHardnesses = SourceMeshDescription->EdgeAttributes().GetAttributes<bool>(MeshAttribute::Edge::IsHard);
+
+	for (const FPolygonID PolygonID : SourceMeshDescription->Polygons().GetElementIDs())
+	{
+		if (ConsumedPolygons[PolygonID.GetValue()])
+		{
+			continue;
+		}
+		TArray<FPolygonID> ConnectedPolygons;
+		TArray<FPolygonID> LastConnectedPolygons;
+		ConnectedPolygons.Add(PolygonID);
+		LastConnectedPolygons.Add(FPolygonID::Invalid);
+		while (ConnectedPolygons.Num() > 0)
+		{
+			check(LastConnectedPolygons.Num() == ConnectedPolygons.Num());
+			FPolygonID LastPolygonID = LastConnectedPolygons.Pop(true);
+			FPolygonID CurrentPolygonID = ConnectedPolygons.Pop(true);
+			if (ConsumedPolygons[CurrentPolygonID.GetValue()])
+			{
+				continue;
+			}
+			TArray<FPolygonID> SoftEdgeNeigbors;
+			uint32& SmoothGroup = PolygonSmoothGroup.FindOrAdd(CurrentPolygonID);
+			uint32 AvoidSmoothGroup = 0;
+			uint32 NeighborSmoothGroup = 0;
+			const uint32 LastSmoothGroupValue = (LastPolygonID == FPolygonID::Invalid) ? 0 : PolygonSmoothGroup[LastPolygonID];
+			TArray<FEdgeID> PolygonEdges;
+			SourceMeshDescription->GetPolygonEdges(CurrentPolygonID, PolygonEdges);
+			for (const FEdgeID& EdgeID : PolygonEdges)
+			{
+				bool bIsHardEdge = EdgeHardnesses[EdgeID];
+				const TArray<FPolygonID>& EdgeConnectedPolygons = SourceMeshDescription->GetEdgeConnectedPolygons(EdgeID);
+				for (const FPolygonID& EdgePolygonID : EdgeConnectedPolygons)
+				{
+					if (EdgePolygonID == CurrentPolygonID)
+					{
+						continue;
+					}
+					uint32 SmoothValue = 0;
+					if (PolygonSmoothGroup.Contains(EdgePolygonID))
+					{
+						SmoothValue = PolygonSmoothGroup[EdgePolygonID];
+					}
+
+					if (bIsHardEdge) //Hard Edge
+					{
+						AvoidSmoothGroup |= SmoothValue;
+					}
+					else
+					{
+						NeighborSmoothGroup |= SmoothValue;
+						//Put all none hard edge polygon in the next iteration
+						if (!ConsumedPolygons[EdgePolygonID.GetValue()])
+						{
+							ConnectedPolygons.Add(EdgePolygonID);
+							LastConnectedPolygons.Add(CurrentPolygonID);
+						}
+						else
+						{
+							SoftEdgeNeigbors.Add(EdgePolygonID);
+						}
+					}
+				}
+			}
+
+			if (AvoidSmoothGroup != 0)
+			{
+				PolygonAvoidances.FindOrAdd(CurrentPolygonID) = AvoidSmoothGroup;
+				//find neighbor avoidance
+				for (FPolygonID& NeighborID : SoftEdgeNeigbors)
+				{
+					if (!PolygonAvoidances.Contains(NeighborID))
+					{
+						continue;
+					}
+					AvoidSmoothGroup |= PolygonAvoidances[NeighborID];
+				}
+				uint32 NewSmoothGroup = 1;
+				while ((NewSmoothGroup & AvoidSmoothGroup) != 0 && NewSmoothGroup < MAX_uint32)
+				{
+					//Shift the smooth group
+					NewSmoothGroup = NewSmoothGroup << 1;
+				}
+				SmoothGroup = NewSmoothGroup;
+				//Apply to all neighboard
+				for (FPolygonID& NeighborID : SoftEdgeNeigbors)
+				{
+					PolygonSmoothGroup[NeighborID] |= NewSmoothGroup;
+				}
+			}
+			else if (NeighborSmoothGroup != 0)
+			{
+				SmoothGroup |= LastSmoothGroupValue | NeighborSmoothGroup;
+			}
+			else
+			{
+				SmoothGroup = 1;
+			}
+			ConsumedPolygons[CurrentPolygonID.GetValue()] = true;
+		}
+	}
+	//Now we have to put the data into the RawMesh
+	int32 TriangleIndex = 0;
+	for (const FPolygonID PolygonID : SourceMeshDescription->Polygons().GetElementIDs())
+	{
+		uint32 PolygonSmoothValue = PolygonSmoothGroup[PolygonID];
+		const TArray<FMeshTriangle>& Triangles = SourceMeshDescription->GetPolygonTriangles(PolygonID);
+		for (const FMeshTriangle& MeshTriangle : Triangles)
+		{
+			DestinationRawMesh.FaceSmoothingMasks[TriangleIndex++] = PolygonSmoothValue;
+		}
+	}
+}
+
+void FMeshDescriptionOperations::ConvertSmoothGroupToHardEdges(const struct FRawMesh &SourceRawMesh, UMeshDescription* DestinationMeshDescription)
+{
+	TEdgeAttributeArray<bool>& EdgeHardnesses = DestinationMeshDescription->EdgeAttributes().GetAttributes<bool>(MeshAttribute::Edge::IsHard);
+
+	TArray<bool> ConsumedPolygons;
+	ConsumedPolygons.AddZeroed(DestinationMeshDescription->Polygons().Num());
+	for (const FPolygonID PolygonID : DestinationMeshDescription->Polygons().GetElementIDs())
+	{
+		if (ConsumedPolygons[PolygonID.GetValue()])
+		{
+			continue;
+		}
+		TArray<FPolygonID> ConnectedPolygons;
+		ConnectedPolygons.Add(PolygonID);
+		while (ConnectedPolygons.Num() > 0)
+		{
+			FPolygonID CurrentPolygonID = ConnectedPolygons.Pop(true);
+			int32 CurrentPolygonIDValue = CurrentPolygonID.GetValue();
+			check(SourceRawMesh.FaceSmoothingMasks.IsValidIndex(CurrentPolygonIDValue));
+			const uint32 ReferenceSmoothGroup = SourceRawMesh.FaceSmoothingMasks[CurrentPolygonIDValue];
+			TArray<FEdgeID> PolygonEdges;
+			DestinationMeshDescription->GetPolygonEdges(CurrentPolygonID, PolygonEdges);
+			for (const FEdgeID& EdgeID : PolygonEdges)
+			{
+				const bool bIsHardEdge = EdgeHardnesses[EdgeID];
+				if (bIsHardEdge)
+				{
+					continue;
+				}
+				const TArray<FPolygonID>& EdgeConnectedPolygons = DestinationMeshDescription->GetEdgeConnectedPolygons(EdgeID);
+				for (const FPolygonID& EdgePolygonID : EdgeConnectedPolygons)
+				{
+					int32 EdgePolygonIDValue = EdgePolygonID.GetValue();
+					if (EdgePolygonID == CurrentPolygonID || ConsumedPolygons[EdgePolygonIDValue])
+					{
+						continue;
+					}
+					check(SourceRawMesh.FaceSmoothingMasks.IsValidIndex(EdgePolygonIDValue));
+					const uint32 TestSmoothGroup = SourceRawMesh.FaceSmoothingMasks[EdgePolygonIDValue];
+					if ((TestSmoothGroup & ReferenceSmoothGroup) == 0)
+					{
+						EdgeHardnesses[EdgeID] = true;
+						break;
+					}
+					else
+					{
+						ConnectedPolygons.Add(EdgePolygonID);
+					}
+				}
+			}
+			ConsumedPolygons[CurrentPolygonID.GetValue()] = true;
+		}
+	}
+}
+
+void FMeshDescriptionOperations::ConverToRawMesh(const UMeshDescription* SourceMeshDescription, struct FRawMesh &DestinationRawMesh)
+{
+	DestinationRawMesh.Empty();
+
+	//Gather all array data
+	const TVertexAttributeArray<FVector>& VertexPositions = SourceMeshDescription->VertexAttributes().GetAttributes<FVector>(MeshAttribute::Vertex::Position);
+
+	const TVertexInstanceAttributeArray<FVector>& VertexInstanceNormals = SourceMeshDescription->VertexInstanceAttributes().GetAttributes<FVector>(MeshAttribute::VertexInstance::Normal);
+	const TVertexInstanceAttributeArray<FVector>& VertexInstanceTangents = SourceMeshDescription->VertexInstanceAttributes().GetAttributes<FVector>(MeshAttribute::VertexInstance::Tangent);
+	const TVertexInstanceAttributeArray<float>& VertexInstanceBinormalSigns = SourceMeshDescription->VertexInstanceAttributes().GetAttributes<float>(MeshAttribute::VertexInstance::BinormalSign);
+	const TVertexInstanceAttributeArray<FVector4>& VertexInstanceColors = SourceMeshDescription->VertexInstanceAttributes().GetAttributes<FVector4>(MeshAttribute::VertexInstance::Color);
+	const TVertexInstanceAttributeIndicesArray<FVector2D>& VertexInstanceUVs = SourceMeshDescription->VertexInstanceAttributes().GetAttributesSet<FVector2D>(MeshAttribute::VertexInstance::TextureCoordinate);
+
+	const TPolygonGroupAttributeArray<int>& PolygonGroupMaterialIndex = SourceMeshDescription->PolygonGroupAttributes().GetAttributes<int>(MeshAttribute::PolygonGroup::MaterialIndex);
+
+	DestinationRawMesh.VertexPositions.AddZeroed(SourceMeshDescription->Vertices().Num());
+	for (const FVertexID& VertexID : SourceMeshDescription->Vertices().GetElementIDs())
+	{
+		int32 VertexIDValue = VertexID.GetValue();
+		DestinationRawMesh.VertexPositions[VertexIDValue] = VertexPositions[VertexID];
+	}
+	int32 VertexInstanceNumber = SourceMeshDescription->VertexInstances().Num();
+	DestinationRawMesh.WedgeColors.AddZeroed(VertexInstanceNumber);
+	DestinationRawMesh.WedgeIndices.AddZeroed(VertexInstanceNumber);
+	DestinationRawMesh.WedgeTangentX.AddZeroed(VertexInstanceNumber);
+	DestinationRawMesh.WedgeTangentY.AddZeroed(VertexInstanceNumber);
+	DestinationRawMesh.WedgeTangentZ.AddZeroed(VertexInstanceNumber);
+	int32 ExistingUVCount = VertexInstanceUVs.GetNumIndices();
+	for (int32 UVIndex = 0; UVIndex < ExistingUVCount; ++UVIndex)
+	{
+		DestinationRawMesh.WedgeTexCoords[UVIndex].AddZeroed(VertexInstanceNumber);
+	}
+
+	int32 TriangleNumber = 0;
+	for (const FPolygonID& PolygonID : SourceMeshDescription->Polygons().GetElementIDs())
+	{
+		TriangleNumber += SourceMeshDescription->GetPolygonTriangles(PolygonID).Num();
+	}
+	DestinationRawMesh.FaceMaterialIndices.AddZeroed(TriangleNumber);
+	DestinationRawMesh.FaceSmoothingMasks.AddZeroed(TriangleNumber);
+
+	for (const FPolygonID& PolygonID : SourceMeshDescription->Polygons().GetElementIDs())
+	{
+		const FPolygonGroupID& PolygonGroupID = SourceMeshDescription->GetPolygonPolygonGroup(PolygonID);
+		int32 PolygonIDValue = PolygonID.GetValue();
+		const TArray<FMeshTriangle>& Triangles = SourceMeshDescription->GetPolygonTriangles(PolygonID);
+		for (const FMeshTriangle& MeshTriangle : Triangles)
+		{
+			DestinationRawMesh.FaceMaterialIndices[PolygonIDValue] = PolygonGroupMaterialIndex[PolygonGroupID];
+			DestinationRawMesh.FaceSmoothingMasks[PolygonIDValue] = 0; //Conversion of soft/hard to smooth mask is done after the geometry is converted
+			for (int32 Corner = 0; Corner < 3; ++Corner)
+			{
+				const FVertexInstanceID VertexInstanceID = MeshTriangle.GetVertexInstanceID(Corner);
+				const int32 VertexInstanceIDValue = VertexInstanceID.GetValue();
+				DestinationRawMesh.WedgeColors[VertexInstanceIDValue] = FLinearColor(VertexInstanceColors[VertexInstanceID]).ToFColor(true);
+				DestinationRawMesh.WedgeIndices[VertexInstanceIDValue] = SourceMeshDescription->GetVertexInstanceVertex(VertexInstanceID).GetValue();
+				DestinationRawMesh.WedgeTangentX[VertexInstanceIDValue] = VertexInstanceTangents[VertexInstanceID];
+				DestinationRawMesh.WedgeTangentY[VertexInstanceIDValue] = FVector::CrossProduct(VertexInstanceNormals[VertexInstanceID], VertexInstanceTangents[VertexInstanceID]).GetSafeNormal() * VertexInstanceBinormalSigns[VertexInstanceID];
+				DestinationRawMesh.WedgeTangentZ[VertexInstanceIDValue] = VertexInstanceNormals[VertexInstanceID];
+				for (int32 UVIndex = 0; UVIndex < ExistingUVCount; ++UVIndex)
+				{
+					DestinationRawMesh.WedgeTexCoords[UVIndex][VertexInstanceIDValue] = VertexInstanceUVs.GetArrayForIndex(UVIndex)[VertexInstanceID];
+				}
+			}
+		}
+	}
+	//Convert the smoothgroup
+	ConvertHardEdgesToSmoothGroup(SourceMeshDescription, DestinationRawMesh);
+}
+
+void FMeshDescriptionOperations::ConverFromRawMesh(const struct FRawMesh &SourceRawMesh, UMeshDescription* DestinationMeshDescription)
+{
+	check(DestinationMeshDescription != nullptr);
+	DestinationMeshDescription->Empty();
+
+	DestinationMeshDescription->ReserveNewVertices(SourceRawMesh.VertexPositions.Num());
+	DestinationMeshDescription->ReserveNewVertexInstances(SourceRawMesh.WedgeIndices.Num());
+	DestinationMeshDescription->ReserveNewPolygons(SourceRawMesh.WedgeIndices.Num() / 3);
+	DestinationMeshDescription->ReserveNewEdges(SourceRawMesh.WedgeIndices.Num() * 2 / 3); //Approximately 2 edges per polygons
+																						   //Gather all array data
+	TVertexAttributeArray<FVector>& VertexPositions = DestinationMeshDescription->VertexAttributes().GetAttributes<FVector>(MeshAttribute::Vertex::Position);
+
+	TVertexInstanceAttributeArray<FVector>& VertexInstanceNormals = DestinationMeshDescription->VertexInstanceAttributes().GetAttributes<FVector>(MeshAttribute::VertexInstance::Normal);
+	TVertexInstanceAttributeArray<FVector>& VertexInstanceTangents = DestinationMeshDescription->VertexInstanceAttributes().GetAttributes<FVector>(MeshAttribute::VertexInstance::Tangent);
+	TVertexInstanceAttributeArray<float>& VertexInstanceBinormalSigns = DestinationMeshDescription->VertexInstanceAttributes().GetAttributes<float>(MeshAttribute::VertexInstance::BinormalSign);
+	TVertexInstanceAttributeArray<FVector4>& VertexInstanceColors = DestinationMeshDescription->VertexInstanceAttributes().GetAttributes<FVector4>(MeshAttribute::VertexInstance::Color);
+	TVertexInstanceAttributeIndicesArray<FVector2D>& VertexInstanceUVs = DestinationMeshDescription->VertexInstanceAttributes().GetAttributesSet<FVector2D>(MeshAttribute::VertexInstance::TextureCoordinate);
+
+	TPolygonGroupAttributeArray<FName>& PolygonGroupImportedMaterialSlotNames = DestinationMeshDescription->PolygonGroupAttributes().GetAttributes<FName>(MeshAttribute::PolygonGroup::ImportedMaterialSlotName);
+	TPolygonGroupAttributeArray<int>& PolygonGroupMaterialIndex = DestinationMeshDescription->PolygonGroupAttributes().GetAttributes<int>(MeshAttribute::PolygonGroup::MaterialIndex);
+
+	int32 NumTexCoords = 0;
+	int32 MaxTexCoords = MAX_MESH_TEXTURE_COORDS;
+	TArray<int32> TextureCoordinnateRemapIndex;
+	TextureCoordinnateRemapIndex.AddZeroed(MaxTexCoords);
+	for (int32 TextureCoordinnateIndex = 0; TextureCoordinnateIndex < MaxTexCoords; ++TextureCoordinnateIndex)
+	{
+		TextureCoordinnateRemapIndex[TextureCoordinnateIndex] = INDEX_NONE;
+		if (SourceRawMesh.WedgeTexCoords[TextureCoordinnateIndex].Num() == SourceRawMesh.WedgeIndices.Num())
+		{
+			TextureCoordinnateRemapIndex[TextureCoordinnateIndex] = NumTexCoords;
+			NumTexCoords++;
+		}
+	}
+	VertexInstanceUVs.SetNumIndices(NumTexCoords);
+	for (int32 VertexIndex = 0; VertexIndex < SourceRawMesh.VertexPositions.Num(); ++VertexIndex)
+	{
+		FVertexID VertexID = DestinationMeshDescription->CreateVertex();
+		VertexPositions[VertexID] = SourceRawMesh.VertexPositions[VertexIndex];
+	}
+
+	bool bHasColors = SourceRawMesh.WedgeColors.Num() > 0;
+	bool bHasTangents = SourceRawMesh.WedgeTangentX.Num() > 0 && SourceRawMesh.WedgeTangentY.Num() > 0;
+	bool bHasNormals = SourceRawMesh.WedgeTangentZ.Num() > 0;
+
+	TArray<FPolygonGroupID> PolygonGroups;
+
+	//Triangles
+	int32 TriangleCount = SourceRawMesh.WedgeIndices.Num() / 3;
+	for (int32 TriangleIndex = 0; TriangleIndex < TriangleCount; ++TriangleIndex)
+	{
+		int32 VerticeIndexBase = TriangleIndex * 3;
+
+		//PolygonGroup
+		FPolygonGroupID PolygonGroupID = FPolygonGroupID::Invalid;
+		int32 MaterialIndex = SourceRawMesh.FaceMaterialIndices[TriangleIndex];
+		for (const FPolygonGroupID& SearchPolygonGroupID : DestinationMeshDescription->PolygonGroups().GetElementIDs())
+		{
+			if (PolygonGroupMaterialIndex[SearchPolygonGroupID] == MaterialIndex)
+			{
+				PolygonGroupID = SearchPolygonGroupID;
+			}
+		}
+		if (PolygonGroupID == FPolygonGroupID::Invalid)
+		{
+			PolygonGroupID = DestinationMeshDescription->CreatePolygonGroup();
+			PolygonGroupMaterialIndex[PolygonGroupID] = MaterialIndex;
+			PolygonGroupImportedMaterialSlotNames[PolygonGroupID] = FName(*FString::Printf(TEXT("MaterialSlot_%d"), MaterialIndex));
+			PolygonGroups.Add(PolygonGroupID);
+		}
+
+		for (int32 Corner = 0; Corner < 3; ++Corner)
+		{
+			int32 VerticeIndex = VerticeIndexBase + Corner;
+			FVertexID VertexID(SourceRawMesh.WedgeIndices[VerticeIndex]);
+			FVertexInstanceID VertexInstanceID = DestinationMeshDescription->CreateVertexInstance(VertexID);
+			VertexInstanceColors[VertexInstanceID] = bHasColors ? FLinearColor::FromSRGBColor(SourceRawMesh.WedgeColors[VerticeIndex]) : FLinearColor::White;
+			VertexInstanceTangents[VertexInstanceID] = bHasTangents ? SourceRawMesh.WedgeTangentX[VerticeIndex] : FVector(ForceInitToZero);
+			VertexInstanceBinormalSigns[VertexInstanceID] = bHasTangents ? GetBasisDeterminantSign(SourceRawMesh.WedgeTangentX[VerticeIndex].GetSafeNormal(), SourceRawMesh.WedgeTangentY[VerticeIndex].GetSafeNormal(), SourceRawMesh.WedgeTangentZ[VerticeIndex].GetSafeNormal()) : 0.0f;
+			VertexInstanceNormals[VertexInstanceID] = bHasNormals ? SourceRawMesh.WedgeTangentZ[VerticeIndex] : FVector(ForceInitToZero);
+			for (int32 TextureCoordinnateIndex = 0; TextureCoordinnateIndex < NumTexCoords; ++TextureCoordinnateIndex)
+			{
+				int32 TextureCoordIndex = TextureCoordinnateRemapIndex[TextureCoordinnateIndex];
+				if (TextureCoordIndex == INDEX_NONE)
+				{
+					continue;
+				}
+				TMeshAttributeArray<FVector2D, FVertexInstanceID>& Uvs = VertexInstanceUVs.GetArrayForIndex(TextureCoordIndex);
+				Uvs[VertexInstanceID] = SourceRawMesh.WedgeTexCoords[TextureCoordinnateIndex][VerticeIndex];
+			}
+		}
+
+		//Create the polygon edges
+		TArray<UMeshDescription::FContourPoint> Contours;
+		// Add the edges of this triangle
+		for (uint32 Corner = 0; Corner < 3; ++Corner)
+		{
+			int32 ContourPointIndex = Contours.AddDefaulted();
+			UMeshDescription::FContourPoint& ContourPoint = Contours[ContourPointIndex];
+			//Find the matching edge ID
+			int32 CornerIndices[2];
+			CornerIndices[0] = (Corner + 0) % 3;
+			CornerIndices[1] = (Corner + 1) % 3;
+
+			FVertexID EdgeVertexIDs[2];
+			EdgeVertexIDs[0] = DestinationMeshDescription->GetVertexInstanceVertex(FVertexInstanceID(VerticeIndexBase + CornerIndices[0]));
+			EdgeVertexIDs[1] = DestinationMeshDescription->GetVertexInstanceVertex(FVertexInstanceID(VerticeIndexBase + CornerIndices[1]));
+
+			FEdgeID MatchEdgeId = DestinationMeshDescription->GetVertexPairEdge(EdgeVertexIDs[0], EdgeVertexIDs[1]);
+			if (MatchEdgeId == FEdgeID::Invalid)
+			{
+				MatchEdgeId = DestinationMeshDescription->CreateEdge(EdgeVertexIDs[0], EdgeVertexIDs[1]);
+			}
+			ContourPoint.EdgeID = MatchEdgeId;
+			ContourPoint.VertexInstanceID = FVertexInstanceID(VerticeIndexBase + CornerIndices[0]);
+		}
+
+		const FPolygonID NewPolygonID = DestinationMeshDescription->CreatePolygon(PolygonGroupID, Contours);
+		int32 NewTriangleIndex = DestinationMeshDescription->GetPolygonTriangles(NewPolygonID).AddDefaulted();
+		FMeshTriangle& NewTriangle = DestinationMeshDescription->GetPolygonTriangles(NewPolygonID)[NewTriangleIndex];
+		for (int32 Corner = 0; Corner < 3; ++Corner)
+		{
+			FVertexInstanceID VertexInstanceID(VerticeIndexBase + Corner);
+			NewTriangle.SetVertexInstanceID(Corner, VertexInstanceID);
+		}
+	}
+	CreatePolygonNTB(DestinationMeshDescription, 0.0f);
+
+	if (!bHasNormals || !bHasTangents)
+	{
+		//Create the missing normals and tangents
+		if (!bHasNormals)
+		{
+			CreateNormals(DestinationMeshDescription, ETangentOptions::BlendOverlappingNormals, false);
+		}
+		CreateMikktTangents(DestinationMeshDescription, ETangentOptions::BlendOverlappingNormals);
+	}
+
+	ConvertSmoothGroupToHardEdges(SourceRawMesh, DestinationMeshDescription);
+}
+
+//////////////////////////////////////////////////////////////////////////
+// Normals tangents and Bi-normals
+
+void FMeshDescriptionOperations::CreatePolygonNTB(UMeshDescription* MeshDescription, float ComparisonThreshold)
+{
+	const TVertexAttributeArray<FVector>& VertexPositions = MeshDescription->VertexAttributes().GetAttributes<FVector>(MeshAttribute::Vertex::Position);
+	TVertexInstanceAttributeArray<FVector2D>& VertexUVs = MeshDescription->VertexInstanceAttributes().GetAttributes<FVector2D>(MeshAttribute::VertexInstance::TextureCoordinate, 0);
+	TPolygonAttributeArray<FVector>& PolygonNormals = MeshDescription->PolygonAttributes().GetAttributes<FVector>(MeshAttribute::Polygon::Normal);
+	TPolygonAttributeArray<FVector>& PolygonTangents = MeshDescription->PolygonAttributes().GetAttributes<FVector>(MeshAttribute::Polygon::Tangent);
+	TPolygonAttributeArray<FVector>& PolygonBinormals = MeshDescription->PolygonAttributes().GetAttributes<FVector>(MeshAttribute::Polygon::Binormal);
+
+	TArray<FPolygonID> DegeneratePolygons;
+
+	FVertexInstanceArray& VertexInstanceArray = MeshDescription->VertexInstances();
+	FVertexArray& VertexArray = MeshDescription->Vertices();
+	FPolygonArray& PolygonArray = MeshDescription->Polygons();
+
+	for (const FPolygonID PolygonID : MeshDescription->Polygons().GetElementIDs())
+	{
+		FVector TangentX(0.0f);
+		FVector TangentY(0.0f);
+		FVector TangentZ(0.0f);
+
+		if (!PolygonNormals[PolygonID].IsNearlyZero())
+		{
+			//By pass normal calculation if its already done
+			continue;
+		}
+		const TArray<FMeshTriangle>& MeshTriangles = MeshDescription->GetPolygonTriangles(PolygonID);
+
+		//We need only the first triangle since all triangle of a polygon must have the same normals (planar polygon)
+		const FMeshTriangle& MeshTriangle = MeshTriangles[0];
+		int32 UVIndex = 0;
+
+		FVector P[3];
+		FVector2D UVs[3];
+
+		for (int32 i = 0; i < 3; ++i)
+		{
+			const FVertexInstanceID VertexInstanceID = MeshTriangle.GetVertexInstanceID(i);
+			UVs[i] = VertexUVs[VertexInstanceID];
+			P[i] = VertexPositions[MeshDescription->GetVertexInstanceVertex(VertexInstanceID)];
+		}
+
+		const FVector Normal = ((P[1] - P[2]) ^ (P[0] - P[2])).GetSafeNormal(ComparisonThreshold);
+		//Check for degenerated polygons, avoid NAN
+		if (!Normal.IsNearlyZero(ComparisonThreshold))
+		{
+			FMatrix	ParameterToLocal(
+				FPlane(P[1].X - P[0].X, P[1].Y - P[0].Y, P[1].Z - P[0].Z, 0),
+				FPlane(P[2].X - P[0].X, P[2].Y - P[0].Y, P[2].Z - P[0].Z, 0),
+				FPlane(P[0].X, P[0].Y, P[0].Z, 0),
+				FPlane(0, 0, 0, 1)
+			);
+
+			FMatrix ParameterToTexture(
+				FPlane(UVs[1].X - UVs[0].X, UVs[1].Y - UVs[0].Y, 0, 0),
+				FPlane(UVs[2].X - UVs[0].X, UVs[2].Y - UVs[0].Y, 0, 0),
+				FPlane(UVs[0].X, UVs[0].Y, 1, 0),
+				FPlane(0, 0, 0, 1)
+			);
+
+			// Use InverseSlow to catch singular matrices.  Inverse can miss this sometimes.
+			const FMatrix TextureToLocal = ParameterToTexture.Inverse() * ParameterToLocal;
+
+			TangentX = TextureToLocal.TransformVector(FVector(1, 0, 0)).GetSafeNormal();
+			TangentY = TextureToLocal.TransformVector(FVector(0, 1, 0)).GetSafeNormal();
+			TangentZ = Normal;
+			FVector::CreateOrthonormalBasis(TangentX, TangentY, TangentZ);
+		}
+		else
+		{
+			DegeneratePolygons.Add(PolygonID);
+		}
+
+		PolygonTangents[PolygonID] = TangentX;
+		PolygonBinormals[PolygonID] = TangentY;
+		PolygonNormals[PolygonID] = TangentZ;
+	}
+
+	//Delete the degenerated polygons. The array is fill only if the remove degenerated option is turn on.
+	if (DegeneratePolygons.Num() > 0)
+	{
+		TArray<FEdgeID> OrphanedEdges;
+		TArray<FVertexInstanceID> OrphanedVertexInstances;
+		TArray<FPolygonGroupID> OrphanedPolygonGroups;
+		TArray<FVertexID> OrphanedVertices;
+		for (FPolygonID& PolygonID : DegeneratePolygons)
+		{
+			MeshDescription->DeletePolygon(PolygonID, &OrphanedEdges, &OrphanedVertexInstances, &OrphanedPolygonGroups);
+		}
+		for (FPolygonGroupID& PolygonGroupID : OrphanedPolygonGroups)
+		{
+			MeshDescription->DeletePolygonGroup(PolygonGroupID);
+		}
+		for (FVertexInstanceID& VertexInstanceID : OrphanedVertexInstances)
+		{
+			MeshDescription->DeleteVertexInstance(VertexInstanceID, &OrphanedVertices);
+		}
+		for (FEdgeID& EdgeID : OrphanedEdges)
+		{
+			MeshDescription->DeleteEdge(EdgeID, &OrphanedVertices);
+		}
+		for (FVertexID& VertexID : OrphanedVertices)
+		{
+			MeshDescription->DeleteVertex(VertexID);
+		}
+		//Compact and Remap IDs so we have clean ID from 0 to n since we just erase some polygons
+		//The render build need to have compact ID
+		FElementIDRemappings RemappingInfos;
+		MeshDescription->Compact(RemappingInfos);
+	}
+}
+
+void FMeshDescriptionOperations::CreateNormals(UMeshDescription* MeshDescription, FMeshDescriptionOperations::ETangentOptions TangentOptions, bool bComputeTangent)
+{
+	//For each vertex compute the normals for every connected edges that are smooth betwween hard edges
+	//         H   A    B
+	//          \  ||  /
+	//       G  -- ** -- C
+	//          // |  \
+		//         F   E    D
+//
+// The double ** are the vertex, the double line are hard edges, the single line are soft edge.
+// A and F are hard, all other edges are soft. The goal is to compute two average normals one from
+// A to F and a second from F to A. Then we can set the vertex instance normals accordingly.
+// First normal(A to F) = Normalize(A+B+C+D+E+F)
+// Second normal(F to A) = Normalize(F+G+H+A)
+// We found the connected edge using the triangle that share edges
+
+// @todo: provide an option to weight each contributing polygon normal according to the size of
+// the angle it makes with the vertex being calculated. This means that triangulated faces whose
+// internal edge meets the vertex doesn't get undue extra weight.
+
+	const TVertexInstanceAttributeArray<FVector2D>& VertexUVs = MeshDescription->VertexInstanceAttributes().GetAttributes<FVector2D>(MeshAttribute::VertexInstance::TextureCoordinate, 0);
+	TVertexInstanceAttributeArray<FVector>& VertexNormals = MeshDescription->VertexInstanceAttributes().GetAttributes<FVector>(MeshAttribute::VertexInstance::Normal, 0);
+	TVertexInstanceAttributeArray<FVector>& VertexTangents = MeshDescription->VertexInstanceAttributes().GetAttributes<FVector>(MeshAttribute::VertexInstance::Tangent, 0);
+	TVertexInstanceAttributeArray<float>& VertexBinormalSigns = MeshDescription->VertexInstanceAttributes().GetAttributes<float>(MeshAttribute::VertexInstance::BinormalSign, 0);
+
+	TPolygonAttributeArray<FVector>& PolygonNormals = MeshDescription->PolygonAttributes().GetAttributes<FVector>(MeshAttribute::Polygon::Normal);
+	TPolygonAttributeArray<FVector>& PolygonTangents = MeshDescription->PolygonAttributes().GetAttributes<FVector>(MeshAttribute::Polygon::Tangent);
+	TPolygonAttributeArray<FVector>& PolygonBinormals = MeshDescription->PolygonAttributes().GetAttributes<FVector>(MeshAttribute::Polygon::Binormal);
+
+	TMap<FPolygonID, FVertexInfo> VertexInfoMap;
+	VertexInfoMap.Reserve(20);
+	//Iterate all vertex to compute normals for all vertex instance
+	for (const FVertexID VertexID : MeshDescription->Vertices().GetElementIDs())
+	{
+		VertexInfoMap.Reset();
+
+		bool bPointHasAllTangents = true;
+		//Fill the VertexInfoMap
+		for (const FEdgeID EdgeID : MeshDescription->GetVertexConnectedEdges(VertexID))
+		{
+			for (const FPolygonID PolygonID : MeshDescription->GetEdgeConnectedPolygons(EdgeID))
+			{
+				FVertexInfo& VertexInfo = VertexInfoMap.FindOrAdd(PolygonID);
+				int32 EdgeIndex = VertexInfo.EdgeIDs.AddUnique(EdgeID);
+				if (VertexInfo.PolygonID == FPolygonID::Invalid)
+				{
+					VertexInfo.PolygonID = PolygonID;
+					for (const FVertexInstanceID VertexInstanceID : MeshDescription->GetPolygonPerimeterVertexInstances(PolygonID))
+					{
+						if (MeshDescription->GetVertexInstanceVertex(VertexInstanceID) == VertexID)
+						{
+							VertexInfo.VertexInstanceID = VertexInstanceID;
+							VertexInfo.UVs = VertexUVs[VertexInstanceID];
+							bPointHasAllTangents &= !VertexNormals[VertexInstanceID].IsNearlyZero() && !VertexTangents[VertexInstanceID].IsNearlyZero();
+							break;
+						}
+					}
+				}
+			}
+		}
+
+		if (bPointHasAllTangents)
+		{
+			continue;
+		}
+
+		//Make sure we consume all our vertex instance
+		check(VertexInfoMap.Num() == MeshDescription->GetVertexVertexInstances(VertexID).Num());
+
+		//Build all group by recursively traverse all polygon connected to the vertex
+		TArray<TArray<FPolygonID>> Groups;
+		TArray<FPolygonID> ConsumedPolygon;
+		for (auto Kvp : VertexInfoMap)
+		{
+			if (ConsumedPolygon.Contains(Kvp.Key))
+			{
+				continue;
+			}
+
+			int32 CurrentGroupIndex = Groups.AddZeroed();
+			TArray<FPolygonID>& CurrentGroup = Groups[CurrentGroupIndex];
+			TArray<FPolygonID> PolygonQueue;
+			PolygonQueue.Add(Kvp.Key); //Use a queue to avoid recursive function
+			while (PolygonQueue.Num() > 0)
+			{
+				FPolygonID CurrentPolygonID = PolygonQueue.Pop(true);
+				FVertexInfo& CurrentVertexInfo = VertexInfoMap.FindOrAdd(CurrentPolygonID);
+				CurrentGroup.AddUnique(CurrentVertexInfo.PolygonID);
+				ConsumedPolygon.AddUnique(CurrentVertexInfo.PolygonID);
+				const TEdgeAttributeArray<bool>& EdgeHardnesses = MeshDescription->EdgeAttributes().GetAttributes<bool>(MeshAttribute::Edge::IsHard);
+				for (const FEdgeID EdgeID : CurrentVertexInfo.EdgeIDs)
+				{
+					if (EdgeHardnesses[EdgeID])
+					{
+						//End of the group
+						continue;
+					}
+					for (const FPolygonID PolygonID : MeshDescription->GetEdgeConnectedPolygons(EdgeID))
+					{
+						if (PolygonID == CurrentVertexInfo.PolygonID)
+						{
+							continue;
+						}
+						//Add this polygon to the group
+						FVertexInfo& OtherVertexInfo = VertexInfoMap.FindOrAdd(PolygonID);
+						//Do not repeat polygons
+						if (!ConsumedPolygon.Contains(OtherVertexInfo.PolygonID))
+						{
+							PolygonQueue.Add(PolygonID);
+						}
+					}
+				}
+			}
+		}
+
+		//Smooth every connected group
+		ConsumedPolygon.Reset();
+		for (const TArray<FPolygonID>& Group : Groups)
+		{
+			//Compute tangents data
+			TMap<FVector2D, FVector> GroupTangent;
+			TMap<FVector2D, FVector> GroupBiNormal;
+
+			TArray<FVertexInstanceID> VertexInstanceInGroup;
+			FVector GroupNormal(0.0f);
+			for (const FPolygonID PolygonID : Group)
+			{
+				ConsumedPolygon.Add(PolygonID);
+				VertexInstanceInGroup.Add(VertexInfoMap[PolygonID].VertexInstanceID);
+				GroupNormal += PolygonNormals[PolygonID];
+				if (bComputeTangent)
+				{
+					const FVector2D UVs = VertexInfoMap[PolygonID].UVs;
+					bool CreateGroup = (!GroupTangent.Contains(UVs));
+					FVector& GroupTangentValue = GroupTangent.FindOrAdd(UVs);
+					FVector& GroupBiNormalValue = GroupBiNormal.FindOrAdd(UVs);
+					GroupTangentValue = CreateGroup ? PolygonTangents[PolygonID] : GroupTangentValue + PolygonTangents[PolygonID];
+					GroupBiNormalValue = CreateGroup ? PolygonBinormals[PolygonID] : GroupBiNormalValue + PolygonBinormals[PolygonID];
+				}
+			}
+
+			//////////////////////////////////////////////////////////////////////////
+			//Apply the group to the Mesh
+			GroupNormal.Normalize();
+			if (bComputeTangent)
+			{
+				for (auto Kvp : GroupTangent)
+				{
+					FVector& GroupTangentValue = GroupTangent.FindOrAdd(Kvp.Key);
+					GroupTangentValue.Normalize();
+				}
+				for (auto Kvp : GroupBiNormal)
+				{
+					FVector& GroupBiNormalValue = GroupBiNormal.FindOrAdd(Kvp.Key);
+					GroupBiNormalValue.Normalize();
+				}
+			}
+			//Apply the average NTB on all Vertex instance
+			for (const FVertexInstanceID VertexInstanceID : VertexInstanceInGroup)
+			{
+				const FVector2D VertexUV = VertexUVs[VertexInstanceID];
+
+				if (VertexNormals[VertexInstanceID].IsNearlyZero(SMALL_NUMBER))
+				{
+					VertexNormals[VertexInstanceID] = GroupNormal;
+				}
+				if (bComputeTangent)
+				{
+					//Avoid changing the original group value
+					FVector GroupTangentValue = GroupTangent[VertexUV];
+					FVector GroupBiNormalValue = GroupBiNormal[VertexUV];
+
+					if (!VertexTangents[VertexInstanceID].IsNearlyZero(SMALL_NUMBER))
+					{
+						GroupTangentValue = VertexTangents[VertexInstanceID];
+					}
+					FVector BiNormal(0.0f);
+					if (!VertexNormals[VertexInstanceID].IsNearlyZero(SMALL_NUMBER) && !VertexTangents[VertexInstanceID].IsNearlyZero(SMALL_NUMBER))
+					{
+						BiNormal = FVector::CrossProduct(VertexNormals[VertexInstanceID], VertexTangents[VertexInstanceID]).GetSafeNormal() * VertexBinormalSigns[VertexInstanceID];
+					}
+					if (!BiNormal.IsNearlyZero(SMALL_NUMBER))
+					{
+						GroupBiNormalValue = BiNormal;
+					}
+					// Gram-Schmidt orthogonalization
+					GroupBiNormalValue -= GroupTangentValue * (GroupTangentValue | GroupBiNormalValue);
+					GroupBiNormalValue.Normalize();
+
+					GroupTangentValue -= VertexNormals[VertexInstanceID] * (VertexNormals[VertexInstanceID] | GroupTangentValue);
+					GroupTangentValue.Normalize();
+
+					GroupBiNormalValue -= VertexNormals[VertexInstanceID] * (VertexNormals[VertexInstanceID] | GroupBiNormalValue);
+					GroupBiNormalValue.Normalize();
+					//Set the value
+					VertexTangents[VertexInstanceID] = GroupTangentValue;
+					//If the BiNormal is zero set the sign to 1.0f
+					VertexBinormalSigns[VertexInstanceID] = GetBasisDeterminantSign(GroupTangentValue, GroupBiNormalValue, VertexNormals[VertexInstanceID]);
+
+				}
+			}
+		}
+	}
+}
+
+namespace MeshDescriptionMikktSpaceInterface
+{
+	//Mikk t spce static function
+	int MikkGetNumFaces(const SMikkTSpaceContext* Context);
+	int MikkGetNumVertsOfFace(const SMikkTSpaceContext* Context, const int FaceIdx);
+	void MikkGetPosition(const SMikkTSpaceContext* Context, float Position[3], const int FaceIdx, const int VertIdx);
+	void MikkGetNormal(const SMikkTSpaceContext* Context, float Normal[3], const int FaceIdx, const int VertIdx);
+	void MikkSetTSpaceBasic(const SMikkTSpaceContext* Context, const float Tangent[3], const float BitangentSign, const int FaceIdx, const int VertIdx);
+	void MikkGetTexCoord(const SMikkTSpaceContext* Context, float UV[2], const int FaceIdx, const int VertIdx);
+}
+
+void FMeshDescriptionOperations::CreateMikktTangents(UMeshDescription* MeshDescription, FMeshDescriptionOperations::ETangentOptions TangentOptions)
+{
+	bool bIgnoreDegenerateTriangles = (TangentOptions & FMeshDescriptionOperations::ETangentOptions::IgnoreDegenerateTriangles) != 0;
+
+	// we can use mikktspace to calculate the tangents
+	SMikkTSpaceInterface MikkTInterface;
+	MikkTInterface.m_getNormal = MeshDescriptionMikktSpaceInterface::MikkGetNormal;
+	MikkTInterface.m_getNumFaces = MeshDescriptionMikktSpaceInterface::MikkGetNumFaces;
+	MikkTInterface.m_getNumVerticesOfFace = MeshDescriptionMikktSpaceInterface::MikkGetNumVertsOfFace;
+	MikkTInterface.m_getPosition = MeshDescriptionMikktSpaceInterface::MikkGetPosition;
+	MikkTInterface.m_getTexCoord = MeshDescriptionMikktSpaceInterface::MikkGetTexCoord;
+	MikkTInterface.m_setTSpaceBasic = MeshDescriptionMikktSpaceInterface::MikkSetTSpaceBasic;
+	MikkTInterface.m_setTSpace = nullptr;
+
+	SMikkTSpaceContext MikkTContext;
+	MikkTContext.m_pInterface = &MikkTInterface;
+	MikkTContext.m_pUserData = (void*)(MeshDescription);
+	MikkTContext.m_bIgnoreDegenerates = bIgnoreDegenerateTriangles;
+	genTangSpaceDefault(&MikkTContext);
+}
+
+namespace MeshDescriptionMikktSpaceInterface
+{
+	int MikkGetNumFaces(const SMikkTSpaceContext* Context)
+	{
+		UMeshDescription *MeshDescription = (UMeshDescription*)(Context->m_pUserData);
+		return MeshDescription->Polygons().Num();
+	}
+
+	int MikkGetNumVertsOfFace(const SMikkTSpaceContext* Context, const int FaceIdx)
+	{
+		// All of our meshes are triangles.
+		return 3;
+	}
+
+	void MikkGetPosition(const SMikkTSpaceContext* Context, float Position[3], const int FaceIdx, const int VertIdx)
+	{
+		UMeshDescription *MeshDescription = (UMeshDescription*)(Context->m_pUserData);
+		const FVertexInstanceID VertexInstanceID(FaceIdx * 3 + VertIdx);
+		const FVertexID VertexID = MeshDescription->GetVertexInstanceVertex(VertexInstanceID);
+		const FVector& VertexPosition = MeshDescription->VertexAttributes().GetAttribute<FVector>(VertexID, MeshAttribute::Vertex::Position);
+		Position[0] = VertexPosition.X;
+		Position[1] = VertexPosition.Y;
+		Position[2] = VertexPosition.Z;
+	}
+
+	void MikkGetNormal(const SMikkTSpaceContext* Context, float Normal[3], const int FaceIdx, const int VertIdx)
+	{
+		UMeshDescription *MeshDescription = (UMeshDescription*)(Context->m_pUserData);
+		const FVertexInstanceID VertexInstanceID(FaceIdx * 3 + VertIdx);
+		const FVector& VertexNormal = MeshDescription->VertexInstanceAttributes().GetAttribute<FVector>(VertexInstanceID, MeshAttribute::VertexInstance::Normal);
+		Normal[0] = VertexNormal.X;
+		Normal[1] = VertexNormal.Y;
+		Normal[2] = VertexNormal.Z;
+	}
+
+	void MikkSetTSpaceBasic(const SMikkTSpaceContext* Context, const float Tangent[3], const float BitangentSign, const int FaceIdx, const int VertIdx)
+	{
+		UMeshDescription *MeshDescription = (UMeshDescription*)(Context->m_pUserData);
+		const FVertexInstanceID VertexInstanceID(FaceIdx * 3 + VertIdx);
+		const FVector VertexTangent(Tangent[0], Tangent[1], Tangent[2]);
+		MeshDescription->VertexInstanceAttributes().SetAttribute<FVector>(VertexInstanceID, MeshAttribute::VertexInstance::Tangent, 0, VertexTangent);
+		MeshDescription->VertexInstanceAttributes().SetAttribute<float>(VertexInstanceID, MeshAttribute::VertexInstance::BinormalSign, 0, -BitangentSign);
+	}
+
+	void MikkGetTexCoord(const SMikkTSpaceContext* Context, float UV[2], const int FaceIdx, const int VertIdx)
+	{
+		UMeshDescription *MeshDescription = (UMeshDescription*)(Context->m_pUserData);
+		const FVertexInstanceID VertexInstanceID(FaceIdx * 3 + VertIdx);
+		const FVector2D& TexCoord = MeshDescription->VertexInstanceAttributes().GetAttribute<FVector2D>(VertexInstanceID, MeshAttribute::VertexInstance::TextureCoordinate, 0);
+		UV[0] = TexCoord.X;
+		UV[1] = TexCoord.Y;
+	}
+}
+
+#undef LOCTEXT_NAMESPACE
diff --git a/Engine/Source/Developer/MeshDescriptionOperations/Public/MeshDescriptionOperations.h b/Engine/Source/Developer/MeshDescriptionOperations/Public/MeshDescriptionOperations.h
new file mode 100644
--- /dev/null
+++ b/Engine/Source/Developer/MeshDescriptionOperations/Public/MeshDescriptionOperations.h
@@ -0,0 +1,62 @@
+// Copyright 1998-2018 Epic Games, Inc. All Rights Reserved.
+
+#pragma once
+
+#include "CoreMinimal.h"
+
+
+class UMeshDescription;
+struct FRawMesh;
+
+DECLARE_LOG_CATEGORY_EXTERN(LogMeshDescriptionOperations, Log, All);
+
+//////////////////////////////////////////////////////////////////////////
+// Any operations on the mesh description that do not depend on engine module
+// should be implement here.
+
+class MESHDESCRIPTIONOPERATIONS_API FMeshDescriptionOperations
+{
+public:
+	enum ETangentOptions
+	{
+		None = 0,
+		BlendOverlappingNormals = 0x1,
+		IgnoreDegenerateTriangles = 0x2,
+		UseMikkTSpace = 0x4,
+	};
+
+	enum class ELightmapUVVersion : int32
+	{
+		BitByBit = 0,
+		Segments = 1,
+		SmallChartPacking = 2,
+		Latest = SmallChartPacking
+	};
+
+	/** Convert this mesh description into the old FRawMesh format. */
+	static void ConverToRawMesh(const class UMeshDescription* SourceMeshDescription, struct FRawMesh &DestinationRawMesh);
+
+	/** Convert old FRawMesh format to MeshDescription. */
+	static void ConverFromRawMesh(const struct FRawMesh &SourceRawMesh, class UMeshDescription* DestinationMeshDescription);
+
+	/** Create some UVs from the specified mesh description data. */
+	static bool GenerateUniqueUVsForStaticMesh(const class UMeshDescription* MeshDescription, int32 TextureResolution, TArray<FVector2D>& OutTexCoords);
+
+	/**
+	 * Compute normal, tangent and Bi-Normal for every polygon in the mesh description. (this do not compute Vertex NTBs)
+	 * It also remove the degenerated polygon from the mesh description
+	 */
+	static void CreatePolygonNTB(UMeshDescription* MeshDescription, float ComparisonThreshold);
+	
+	/** Compute normal, tangent and Bi-Normal(only if bComputeTangent is true) for every vertex in the mesh description. */
+	static void CreateNormals(UMeshDescription* MeshDescription, ETangentOptions TangentOptions, bool bComputeTangent);
+	
+	/** Compute tangent and Bi-Normal using mikkt space for every vertex in the mesh description. */
+	static void CreateMikktTangents(UMeshDescription* MeshDescription, ETangentOptions TangentOptions);
+
+private:
+	
+	static void ConvertHardEdgesToSmoothGroup(const UMeshDescription* SourceMeshDescription, struct FRawMesh &DestinationRawMesh);
+
+	static void ConvertSmoothGroupToHardEdges(const struct FRawMesh &SourceRawMesh, UMeshDescription* DestinationMeshDescription);
+};
\ No newline at end of file
diff --git a/Engine/Source/Runtime/Engine/Engine.Build.cs b/Engine/Source/Runtime/Engine/Engine.Build.cs
--- a/Engine/Source/Runtime/Engine/Engine.Build.cs
+++ b/Engine/Source/Runtime/Engine/Engine.Build.cs
@@ -117,8 +117,8 @@ public class Engine : ModuleRules
 		{
 			// for now we depend on this
 			PrivateDependencyModuleNames.Add("RawMesh");
-			PrivateDependencyModuleNames.Add("MeshBuilder");
-		}
+            PrivateDependencyModuleNames.Add("MeshDescriptionOperations");
+        }
 
 		bool bVariadicTemplatesSupported = true;
 		if (Target.Platform == UnrealTargetPlatform.XboxOne)
diff --git a/Engine/Source/Runtime/Engine/Private/StaticMesh.cpp b/Engine/Source/Runtime/Engine/Private/StaticMesh.cpp
--- a/Engine/Source/Runtime/Engine/Private/StaticMesh.cpp
+++ b/Engine/Source/Runtime/Engine/Private/StaticMesh.cpp
@@ -50,6 +50,7 @@
 #include "MeshUtilities.h"
 #include "DerivedDataCacheInterface.h"
 #include "IMeshBuilderModule.h"
+#include "MeshDescriptionOperations.h"
 #endif // #if WITH_EDITOR
 
 #include "Engine/StaticMeshSocket.h"
diff --git a/Engine/Source/Runtime/MeshBuilder/MeshBuilder.Build.cs b/Engine/Source/Runtime/MeshBuilder/MeshBuilder.Build.cs
--- a/Engine/Source/Runtime/MeshBuilder/MeshBuilder.Build.cs
+++ b/Engine/Source/Runtime/MeshBuilder/MeshBuilder.Build.cs
@@ -17,6 +17,7 @@ namespace UnrealBuildTool.Rules
                     "Engine",
                     "RenderCore",
                     "MeshDescription",
+                    "MeshDescriptionOperations",
                     "MeshReductionInterface",
                     "RenderCore",
                     "RawMesh"
@@ -26,7 +27,6 @@ namespace UnrealBuildTool.Rules
             AddEngineThirdPartyPrivateStaticDependencies(Target, "nvTriStrip");
             AddEngineThirdPartyPrivateStaticDependencies(Target, "ForsythTriOptimizer");
             AddEngineThirdPartyPrivateStaticDependencies(Target, "nvTessLib");
-            AddEngineThirdPartyPrivateStaticDependencies(Target, "MikkTSpace");
             AddEngineThirdPartyPrivateStaticDependencies(Target, "QuadricMeshReduction");
 
             if (Target.bCompileSimplygon == true)
diff --git a/Engine/Source/Runtime/MeshBuilder/Private/LayoutUV.cpp b/Engine/Source/Runtime/MeshBuilder/Private/LayoutUV.cpp
--- a/Engine/Source/Runtime/MeshBuilder/Private/LayoutUV.cpp
+++ b/Engine/Source/Runtime/MeshBuilder/Private/LayoutUV.cpp
@@ -2,6 +2,7 @@
 
 #include "LayoutUV.h"
 #include "DisjointSet.h"
+#include "MeshDescriptionOperations.h"
 
 #include "Algo/IntroSort.h"
 
@@ -19,7 +20,7 @@ FLayoutUV::FLayoutUV( UMeshDescription* InMesh, uint32 InSrcChannel, uint32 InDs
 	, ChartRaster( TextureResolution, TextureResolution )
 	, BestChartRaster( TextureResolution, TextureResolution )
 	, ChartShader( &ChartRaster )
-	, LayoutVersion(FMeshDescriptionHelper::ELightmapUVVersion::Latest )
+	, LayoutVersion(FMeshDescriptionOperations::ELightmapUVVersion::Latest )
 {}
 
 void FLayoutUV::FindCharts( const TMultiMap<int32,int32>& OverlappingCorners )
@@ -218,7 +219,7 @@ void FLayoutUV::FindCharts( const TMultiMap<int32,int32>& OverlappingCorners )
 		Chart.LastTri = Tri;
 
 #if !CHART_JOINING
-		if (LayoutVersion >= FMeshDescriptionHelper::ELightmapUVVersion::SmallChartPacking)
+		if (LayoutVersion >= FMeshDescriptionOperations::ELightmapUVVersion::SmallChartPacking)
 		{
 			Chart.WorldScale /= FMath::Max(Chart.UVArea, 1e-8f);
 		}
@@ -537,7 +538,7 @@ void FLayoutUV::FindCharts( const TMultiMap<int32,int32>& OverlappingCorners )
 	{
 		FMeshChart& Chart = Charts[i];
 
-		if (LayoutVersion >= FMeshDescriptionHelper::ELightmapUVVersion::SmallChartPacking)
+		if (LayoutVersion >= FMeshDescriptionOperations::ELightmapUVVersion::SmallChartPacking)
 		{
 			Chart.WorldScale /= FMath::Max(Chart.UVArea, 1e-8f);
 		}
@@ -809,11 +810,11 @@ bool FLayoutUV::PackCharts()
 			}
 			else
 			{
-				if ( LayoutVersion >= FMeshDescriptionHelper::ELightmapUVVersion::Segments && Orientation % 4 == 1 )
+				if ( LayoutVersion >= FMeshDescriptionOperations::ELightmapUVVersion::Segments && Orientation % 4 == 1 )
 				{
 					ChartRaster.FlipX( Rect );
 				}
-				else if ( LayoutVersion >= FMeshDescriptionHelper::ELightmapUVVersion::Segments && Orientation % 4 == 3 )
+				else if ( LayoutVersion >= FMeshDescriptionOperations::ELightmapUVVersion::Segments && Orientation % 4 == 3 )
 				{
 					ChartRaster.FlipY( Rect );
 				}
@@ -827,11 +828,11 @@ bool FLayoutUV::PackCharts()
 				bool bFound = false;
 
 				uint32 BeginFind = FPlatformTime::Cycles();
-				if ( LayoutVersion == FMeshDescriptionHelper::ELightmapUVVersion::BitByBit )
+				if ( LayoutVersion == FMeshDescriptionOperations::ELightmapUVVersion::BitByBit )
 				{
 					bFound = LayoutRaster.FindBitByBit( Rect, ChartRaster );
 				}
-				else if ( LayoutVersion >= FMeshDescriptionHelper::ELightmapUVVersion::Segments )
+				else if ( LayoutVersion >= FMeshDescriptionOperations::ELightmapUVVersion::Segments )
 				{
 					bFound = LayoutRaster.FindWithSegments( Rect, BestRect, ChartRaster );
 				}
@@ -1044,7 +1045,7 @@ void FLayoutUV::RasterizeChart( const FMeshChart& Chart, uint32 RectW, uint32 Re
 		RasterizeTriangle< FAllocator2DShader, 16 >( ChartShader, Points, RectW, RectH );
 	}
 
-	if ( LayoutVersion >= FMeshDescriptionHelper::ELightmapUVVersion::Segments )
+	if ( LayoutVersion >= FMeshDescriptionOperations::ELightmapUVVersion::Segments )
 	{
 		ChartRaster.CreateUsedSegments();
 	}
diff --git a/Engine/Source/Runtime/MeshBuilder/Private/LayoutUV.h b/Engine/Source/Runtime/MeshBuilder/Private/LayoutUV.h
--- a/Engine/Source/Runtime/MeshBuilder/Private/LayoutUV.h
+++ b/Engine/Source/Runtime/MeshBuilder/Private/LayoutUV.h
@@ -6,6 +6,7 @@
 #include "Allocator2D.h"
 #include "MeshDescription.h"
 #include "MeshDescriptionHelper.h"
+#include "MeshDescriptionOperations.h"
 #include "MeshAttributes.h"
 
 #define NEW_UVS_ARE_SAME THRESH_POINTS_ARE_SAME
@@ -53,7 +54,7 @@ public:
 	bool		FindBestPacking();
 	void		CommitPackedUVs();
 
-	void		SetVersion(FMeshDescriptionHelper::ELightmapUVVersion Version ) { LayoutVersion = Version; }
+	void		SetVersion(FMeshDescriptionOperations::ELightmapUVVersion Version ) { LayoutVersion = Version; }
 
 private:
 	bool		PositionsMatch( uint32 a, uint32 b ) const;
@@ -68,7 +69,7 @@ private:
 	void		OrientChart( FMeshChart& Chart, int32 Orientation );
 	void		RasterizeChart( const FMeshChart& Chart, uint32 RectW, uint32 RectH );
 
-	float		GetUVEqualityThreshold() const { return LayoutVersion >= FMeshDescriptionHelper::ELightmapUVVersion::SmallChartPacking ? NEW_UVS_ARE_SAME : LEGACY_UVS_ARE_SAME; }
+	float		GetUVEqualityThreshold() const { return LayoutVersion >= FMeshDescriptionOperations::ELightmapUVVersion::SmallChartPacking ? NEW_UVS_ARE_SAME : LEGACY_UVS_ARE_SAME; }
 
 	UMeshDescription* MeshDescription;
 	uint32		SrcChannel;
@@ -86,7 +87,7 @@ private:
 	FAllocator2D		BestChartRaster;
 	FAllocator2DShader	ChartShader;
 
-	FMeshDescriptionHelper::ELightmapUVVersion LayoutVersion;
+	FMeshDescriptionOperations::ELightmapUVVersion LayoutVersion;
 };
 
 
diff --git a/Engine/Source/Runtime/MeshBuilder/Private/MeshBuilder.cpp b/Engine/Source/Runtime/MeshBuilder/Private/MeshBuilder.cpp
--- a/Engine/Source/Runtime/MeshBuilder/Private/MeshBuilder.cpp
+++ b/Engine/Source/Runtime/MeshBuilder/Private/MeshBuilder.cpp
@@ -16,77 +16,7 @@ FMeshBuilder::FMeshBuilder()
 
 }
 
-void FMeshDescriptionOperations::ComputeMeshNTBs(class UMeshDescription* MeshDescription, const struct FMeshBuildSettings& BuildSettings)
-{
-	FVertexInstanceArray& VertexInstanceArray = MeshDescription->VertexInstances();
-	TVertexInstanceAttributeArray<FVector>& Normals = MeshDescription->VertexInstanceAttributes().GetAttributes<FVector>(MeshAttribute::VertexInstance::Normal);
-	TVertexInstanceAttributeArray<FVector>& Tangents = MeshDescription->VertexInstanceAttributes().GetAttributes<FVector>(MeshAttribute::VertexInstance::Tangent);
-	TVertexInstanceAttributeArray<float>& BinormalSigns = MeshDescription->VertexInstanceAttributes().GetAttributes<float>(MeshAttribute::VertexInstance::BinormalSign);
-
-	// Static meshes always blend normals of overlapping corners.
-	uint32 TangentOptions = FMeshDescriptionHelper::ETangentOptions::BlendOverlappingNormals;
-	if (BuildSettings.bRemoveDegenerates)
-	{
-		// If removing degenerate triangles, ignore them when computing tangents.
-		TangentOptions |= FMeshDescriptionHelper::ETangentOptions::IgnoreDegenerateTriangles;
-	}
-
-	//This function make sure the Polygon NTB are compute and also remove degenerated triangle from the render mesh description.
-	FMeshDescriptionHelper::CreatePolygonNTB(MeshDescription, BuildSettings.bRemoveDegenerates ? SMALL_NUMBER : 0.0f);
-
-	//Keep the original mesh description NTBs if we do not rebuild the normals or tangents.
-	bool bComputeTangentLegacy = !BuildSettings.bUseMikkTSpace && (BuildSettings.bRecomputeNormals || BuildSettings.bRecomputeTangents);
-	bool bHasAllNormals = true;
-	bool bHasAllTangents = true;
-	for (const FVertexInstanceID VertexInstanceID : VertexInstanceArray.GetElementIDs())
-	{
-		// Dump normals and tangents if we are recomputing them.
-		if (BuildSettings.bRecomputeTangents)
-		{
-			//Dump the tangents
-			BinormalSigns[VertexInstanceID] = 0.0f;
-			Tangents[VertexInstanceID] = FVector(0.0f);
-		}
-		if (BuildSettings.bRecomputeNormals)
-		{
-			//Dump the normals
-			Normals[VertexInstanceID] = FVector(0.0f);
-		}
-		bHasAllNormals &= !Normals[VertexInstanceID].IsNearlyZero();
-		bHasAllTangents &= !Tangents[VertexInstanceID].IsNearlyZero();
-	}
-
-
-	//MikkTSpace should be use only when the user want to recompute the normals or tangents otherwise should always fallback on builtin
-	//We cannot use mikkt space with degenerated normals fallback on buitin.
-	if (BuildSettings.bUseMikkTSpace && (BuildSettings.bRecomputeNormals || BuildSettings.bRecomputeTangents))
-	{
-		if (!bHasAllNormals)
-		{
-			FMeshDescriptionHelper::CreateNormals(MeshDescription, (FMeshDescriptionHelper::ETangentOptions)TangentOptions, false);
-		}
-		FMeshDescriptionHelper::CreateMikktTangents(MeshDescription, (FMeshDescriptionHelper::ETangentOptions)TangentOptions);
-	}
-	else if (!bHasAllNormals || !bHasAllTangents)
-	{
-		//Set the compute tangent to true when we do not build using mikkt space
-		FMeshDescriptionHelper::CreateNormals(MeshDescription, (FMeshDescriptionHelper::ETangentOptions)TangentOptions, true);
-	}
-}
-
-/** Convert this mesh description into the old FRawMesh format*/
-void FMeshDescriptionOperations::ConverToRawMesh(const UMeshDescription* SourceMeshDescription, FRawMesh &DestinationRawMesh)
-{
-	FMeshDescriptionHelper::ConverToRawMesh(SourceMeshDescription, DestinationRawMesh);
-}
-
-/** Convert old FRawMesh format to MeshDescription*/
-void FMeshDescriptionOperations::ConverFromRawMesh(const FRawMesh &SourceRawMesh, UMeshDescription* DestinationMeshDescription)
-{
-	FMeshDescriptionHelper::ConverFromRawMesh(SourceRawMesh, DestinationMeshDescription);
-}
-
-bool FMeshDescriptionOperations::GenerateUniqueUVsForStaticMesh(const UMeshDescription* MeshDescription, int32 TextureResolution, TArray<FVector2D>& OutTexCoords)
+bool FMeshBuilderOperations::GenerateUniqueUVsForStaticMesh(const UMeshDescription* MeshDescription, int32 TextureResolution, TArray<FVector2D>& OutTexCoords)
 {
 	// Create a copy of original mesh (only copy necessary data)
 	UMeshDescription* DuplicateMeshDescription = Cast<UMeshDescription>(StaticDuplicateObject(MeshDescription, GetTransientPackage(), NAME_None, RF_NoFlags));
diff --git a/Engine/Source/Runtime/MeshBuilder/Private/MeshDescriptionHelper.cpp b/Engine/Source/Runtime/MeshBuilder/Private/MeshDescriptionHelper.cpp
--- a/Engine/Source/Runtime/MeshBuilder/Private/MeshDescriptionHelper.cpp
+++ b/Engine/Source/Runtime/MeshBuilder/Private/MeshDescriptionHelper.cpp
@@ -10,7 +10,6 @@
 #include "Engine/StaticMesh.h"
 #include "MeshDescription.h"
 #include "MeshAttributes.h"
-#include "mikktspace.h"
 #include "Serialization/MemoryWriter.h"
 #include "Serialization/MemoryReader.h"
 #include "RenderUtils.h"
@@ -20,23 +19,13 @@
 #include "Materials/Material.h"
 #include "RawMesh.h"
 #include "BuildStatisticManager.h"
+#include "MeshDescriptionOperations.h"
 
 //Enable all check
 //#define ENABLE_NTB_CHECK
 
 DEFINE_LOG_CATEGORY(LogMeshDescriptionBuildStatistic);
 
-namespace MeshDescriptionMikktSpaceInterface
-{
-	//Mikk t spce static function
-	int MikkGetNumFaces(const SMikkTSpaceContext* Context);
-	int MikkGetNumVertsOfFace(const SMikkTSpaceContext* Context, const int FaceIdx);
-	void MikkGetPosition(const SMikkTSpaceContext* Context, float Position[3], const int FaceIdx, const int VertIdx);
-	void MikkGetNormal(const SMikkTSpaceContext* Context, float Normal[3], const int FaceIdx, const int VertIdx);
-	void MikkSetTSpaceBasic(const SMikkTSpaceContext* Context, const float Tangent[3], const float BitangentSign, const int FaceIdx, const int VertIdx);
-	void MikkGetTexCoord(const SMikkTSpaceContext* Context, float UV[2], const int FaceIdx, const int VertIdx);
-}
-
 FMeshDescriptionHelper::FMeshDescriptionHelper(FMeshBuildSettings* InBuildSettings, const UMeshDescription* InOriginalMeshDescription)
 	: OriginalMeshDescription(InOriginalMeshDescription)
 	, BuildSettings(InBuildSettings)
@@ -70,15 +59,15 @@ UMeshDescription* FMeshDescriptionHelper::GetRenderMeshDescription(UObject* Owne
 	// Compute any missing normals or tangents.
 	{
 		// Static meshes always blend normals of overlapping corners.
-		uint32 TangentOptions = FMeshDescriptionHelper::ETangentOptions::BlendOverlappingNormals;
+		uint32 TangentOptions = FMeshDescriptionOperations::ETangentOptions::BlendOverlappingNormals;
 		if (BuildSettings->bRemoveDegenerates)
 		{
 			// If removing degenerate triangles, ignore them when computing tangents.
-			TangentOptions |= FMeshDescriptionHelper::ETangentOptions::IgnoreDegenerateTriangles;
+			TangentOptions |= FMeshDescriptionOperations::ETangentOptions::IgnoreDegenerateTriangles;
 		}
 
 		//This function make sure the Polygon NTB are compute and also remove degenerated triangle from the render mesh description.
-		CreatePolygonNTB(RenderMeshDescription, BuildSettings->bRemoveDegenerates ? SMALL_NUMBER : 0.0f);
+		FMeshDescriptionOperations::CreatePolygonNTB(RenderMeshDescription, BuildSettings->bRemoveDegenerates ? SMALL_NUMBER : 0.0f);
 		
 		//Keep the original mesh description NTBs if we do not rebuild the normals or tangents.
 		bool bComputeTangentLegacy = !BuildSettings->bUseMikkTSpace && (BuildSettings->bRecomputeNormals || BuildSettings->bRecomputeTangents);
@@ -109,14 +98,14 @@ UMeshDescription* FMeshDescriptionHelper::GetRenderMeshDescription(UObject* Owne
 		{
 			if (!bHasAllNormals)
 			{
-				CreateNormals(RenderMeshDescription, (FMeshDescriptionHelper::ETangentOptions)TangentOptions, false);
+				FMeshDescriptionOperations::CreateNormals(RenderMeshDescription, (FMeshDescriptionOperations::ETangentOptions)TangentOptions, false);
 			}
-			CreateMikktTangents(RenderMeshDescription, (FMeshDescriptionHelper::ETangentOptions)TangentOptions);
+			FMeshDescriptionOperations::CreateMikktTangents(RenderMeshDescription, (FMeshDescriptionOperations::ETangentOptions)TangentOptions);
 		}
 		else if(!bHasAllNormals || !bHasAllTangents)
 		{
 			//Set the compute tangent to true when we do not build using mikkt space
-			CreateNormals(RenderMeshDescription, (FMeshDescriptionHelper::ETangentOptions)TangentOptions, true);
+			FMeshDescriptionOperations::CreateNormals(RenderMeshDescription, (FMeshDescriptionOperations::ETangentOptions)TangentOptions, true);
 		}
 	}
 
@@ -144,7 +133,7 @@ UMeshDescription* FMeshDescriptionHelper::GetRenderMeshDescription(UObject* Owne
 		}
 
 		FLayoutUV Packer(RenderMeshDescription, BuildSettings->SrcLightmapIndex, BuildSettings->DstLightmapIndex, BuildSettings->MinLightmapResolution);
-		Packer.SetVersion((FMeshDescriptionHelper::ELightmapUVVersion)(StaticMesh->LightmapUVVersion));
+		Packer.SetVersion((FMeshDescriptionOperations::ELightmapUVVersion)(StaticMesh->LightmapUVVersion));
 
 		Packer.FindCharts(OverlappingCorners);
 		bool bPackSuccess = Packer.FindBestPacking();
@@ -232,910 +221,3 @@ void FMeshDescriptionHelper::FindOverlappingCorners(const UMeshDescription* Mesh
 	FindOverlappingCorners(OverlappingCorners, MeshDescription, ComparisonThreshold);
 }
 
-void FMeshDescriptionHelper::CreatePolygonNTB(UMeshDescription* MeshDescription, float ComparisonThreshold)
-{
-	const TVertexAttributeArray<FVector>& VertexPositions = MeshDescription->VertexAttributes().GetAttributes<FVector>(MeshAttribute::Vertex::Position);
-	TVertexInstanceAttributeArray<FVector2D>& VertexUVs = MeshDescription->VertexInstanceAttributes().GetAttributes<FVector2D>(MeshAttribute::VertexInstance::TextureCoordinate, 0);
-	TPolygonAttributeArray<FVector>& PolygonNormals = MeshDescription->PolygonAttributes().GetAttributes<FVector>(MeshAttribute::Polygon::Normal);
-	TPolygonAttributeArray<FVector>& PolygonTangents = MeshDescription->PolygonAttributes().GetAttributes<FVector>(MeshAttribute::Polygon::Tangent);
-	TPolygonAttributeArray<FVector>& PolygonBinormals = MeshDescription->PolygonAttributes().GetAttributes<FVector>(MeshAttribute::Polygon::Binormal);
-
-	TArray<FPolygonID> DegeneratePolygons;
-
-	FVertexInstanceArray& VertexInstanceArray = MeshDescription->VertexInstances();
-	FVertexArray& VertexArray = MeshDescription->Vertices();
-	FPolygonArray& PolygonArray = MeshDescription->Polygons();
-
-	for (const FPolygonID PolygonID : MeshDescription->Polygons().GetElementIDs())
-	{
-		FVector TangentX(0.0f);
-		FVector TangentY(0.0f);
-		FVector TangentZ(0.0f);
-
-		if (!PolygonNormals[PolygonID].IsNearlyZero())
-		{
-			//By pass normal calculation if its already done
-			continue;
-		}
-		const TArray<FMeshTriangle>& MeshTriangles = MeshDescription->GetPolygonTriangles(PolygonID);
-#ifdef ENABLE_NTB_CHECK
-		//Assume triangle are build
-		check(MeshTriangles.Num() > 0);
-#endif
-
-		//We need only the first triangle since all triangle of a polygon must have the same normals (planar polygon)
-		const FMeshTriangle& MeshTriangle = MeshTriangles[0];
-		int32 UVIndex = 0;
-
-		FVector P[3];
-		FVector2D UVs[3];
-
-		for (int32 i = 0; i < 3; ++i)
-		{
-			const FVertexInstanceID VertexInstanceID = MeshTriangle.GetVertexInstanceID(i);
-			UVs[i] = VertexUVs[VertexInstanceID];
-			P[i] = VertexPositions[MeshDescription->GetVertexInstanceVertex(VertexInstanceID)];
-		}
-
-		const FVector Normal = ((P[1] - P[2]) ^ (P[0] - P[2])).GetSafeNormal(ComparisonThreshold);
-		//Check for degenerated polygons, avoid NAN
-		if (!Normal.IsNearlyZero(ComparisonThreshold))
-		{
-			FMatrix	ParameterToLocal(
-				FPlane(P[1].X - P[0].X, P[1].Y - P[0].Y, P[1].Z - P[0].Z, 0),
-				FPlane(P[2].X - P[0].X, P[2].Y - P[0].Y, P[2].Z - P[0].Z, 0),
-				FPlane(P[0].X, P[0].Y, P[0].Z, 0),
-				FPlane(0, 0, 0, 1)
-			);
-
-			FMatrix ParameterToTexture(
-				FPlane(UVs[1].X - UVs[0].X, UVs[1].Y - UVs[0].Y, 0, 0),
-				FPlane(UVs[2].X - UVs[0].X, UVs[2].Y - UVs[0].Y, 0, 0),
-				FPlane(UVs[0].X, UVs[0].Y, 1, 0),
-				FPlane(0, 0, 0, 1)
-			);
-
-			// Use InverseSlow to catch singular matrices.  Inverse can miss this sometimes.
-			const FMatrix TextureToLocal = ParameterToTexture.Inverse() * ParameterToLocal;
-
-			TangentX = TextureToLocal.TransformVector(FVector(1, 0, 0)).GetSafeNormal();
-			TangentY = TextureToLocal.TransformVector(FVector(0, 1, 0)).GetSafeNormal();
-			TangentZ = Normal;
-			FVector::CreateOrthonormalBasis(TangentX, TangentY, TangentZ);
-		}
-		else
-		{
-			DegeneratePolygons.Add(PolygonID);
-		}
-
-		PolygonTangents[PolygonID] = TangentX;
-		PolygonBinormals[PolygonID] = TangentY;
-		PolygonNormals[PolygonID] = TangentZ;
-	}
-
-	//Delete the degenerated polygons. The array is fill only if the remove degenerated option is turn on.
-	if (DegeneratePolygons.Num() > 0)
-	{
-		TArray<FEdgeID> OrphanedEdges;
-		TArray<FVertexInstanceID> OrphanedVertexInstances;
-		TArray<FPolygonGroupID> OrphanedPolygonGroups;
-		TArray<FVertexID> OrphanedVertices;
-		for (FPolygonID& PolygonID : DegeneratePolygons)
-		{
-			MeshDescription->DeletePolygon(PolygonID, &OrphanedEdges, &OrphanedVertexInstances, &OrphanedPolygonGroups);
-		}
-		for (FPolygonGroupID& PolygonGroupID : OrphanedPolygonGroups)
-		{
-			MeshDescription->DeletePolygonGroup(PolygonGroupID);
-		}
-		for (FVertexInstanceID& VertexInstanceID : OrphanedVertexInstances)
-		{
-			MeshDescription->DeleteVertexInstance(VertexInstanceID, &OrphanedVertices);
-		}
-		for (FEdgeID& EdgeID : OrphanedEdges)
-		{
-			MeshDescription->DeleteEdge(EdgeID, &OrphanedVertices);
-		}
-		for (FVertexID& VertexID : OrphanedVertices)
-		{
-			MeshDescription->DeleteVertex(VertexID);
-		}
-		//Compact and Remap IDs so we have clean ID from 0 to n since we just erase some polygons
-		//The render build need to have compact ID
-		FElementIDRemappings RemappingInfos;
-		MeshDescription->Compact(RemappingInfos);
-	}
-}
-
-struct FVertexInfo
-{
-	FVertexInfo()
-	{
-		PolygonID = FPolygonID::Invalid;
-		VertexInstanceID = FVertexInstanceID::Invalid;
-		UVs = FVector2D(0.0f, 0.0f);
-		EdgeIDs.Reserve(2);//Most of the time a edge has two triangles
-	}
-
-	FPolygonID PolygonID;
-	FVertexInstanceID VertexInstanceID;
-	FVector2D UVs;
-	TArray<FEdgeID> EdgeIDs;
-};
-
-namespace MeshDescriptionPrototype
-{
-	void GetVertexConnectedPolygons(const UMeshDescription* MeshDescription, const FVertexID VertexID, TArray<FPolygonID>& OutConnectedPolygonIDs)
-	{
-		OutConnectedPolygonIDs.Reset();
-
-		const FVertexInstanceArray& VertexInstances = MeshDescription->VertexInstances();
-		for (const FVertexInstanceID VertexInstanceID : MeshDescription->GetVertex(VertexID).VertexInstanceIDs)
-		{
-			OutConnectedPolygonIDs.Append(VertexInstances[VertexInstanceID].ConnectedPolygons);
-		}
-	}
-
-	void GetConnectedSoftEdges(const UMeshDescription* MeshDescription, const FVertexID VertexID, TArray<FEdgeID>& OutConnectedSoftEdges)
-	{
-		OutConnectedSoftEdges.Reset();
-
-		const TEdgeAttributeArray<bool>& EdgeHardnesses = MeshDescription->EdgeAttributes().GetAttributes<bool>(MeshAttribute::Edge::IsHard);
-		for (const FEdgeID ConnectedEdgeID : MeshDescription->GetVertex(VertexID).ConnectedEdgeIDs)
-		{
-			if (!EdgeHardnesses[ConnectedEdgeID])
-			{
-				OutConnectedSoftEdges.Add(ConnectedEdgeID);
-			}
-		}
-	}
-
-	void GetPolygonsInSameSoftEdgedGroup(const UMeshDescription* MeshDescription, const FVertexID VertexID, const FPolygonID PolygonID, TArray<FPolygonID>& OutPolygonIDs)
-	{
-		// The aim here is to determine which polygons form part of the same soft edged group as the polygons attached to this vertex instance.
-		// They should all contribute to the final vertex instance normal.
-
-		OutPolygonIDs.Reset();
-
-		// Get all polygons connected to this vertex.
-		static TArray<FPolygonID> ConnectedPolygons;
-		GetVertexConnectedPolygons(MeshDescription, VertexID, ConnectedPolygons);
-
-		// Cache a list of all soft edges which share this vertex.
-		// We're only interested in finding adjacent polygons which are not the other side of a hard edge.
-		static TArray<FEdgeID> ConnectedSoftEdges;
-		GetConnectedSoftEdges(MeshDescription, VertexID, ConnectedSoftEdges);
-
-		// Iterate through adjacent polygons which contain the given vertex, but don't cross a hard edge.
-		// Maintain a list of polygon IDs to be examined. Adjacents are added to the list if suitable.
-		// Add the start poly here.
-		static TArray<FPolygonID> PolygonsToCheck;
-		PolygonsToCheck.Reset();
-		PolygonsToCheck.Add(PolygonID);
-
-		const FEdgeArray& Edges = MeshDescription->Edges();
-		int32 Index = 0;
-		while (Index < PolygonsToCheck.Num())
-		{
-			const FPolygonID PolygonToCheck = PolygonsToCheck[Index];
-			Index++;
-
-			if (ConnectedPolygons.Contains(PolygonToCheck))
-			{
-				OutPolygonIDs.Add(PolygonToCheck);
-
-				// Now look at its adjacent polygons. If they are joined by a soft edge which includes the vertex we're interested in, we want to consider them.
-				// We take a shortcut by doing this process in reverse: we already know all the soft edges we are interested in, so check if any of them
-				// have the current polygon as an adjacent.
-				for (const FEdgeID ConnectedSoftEdge : ConnectedSoftEdges)
-				{
-					const FMeshEdge& Edge = Edges[ConnectedSoftEdge];
-					if (Edge.ConnectedPolygons.Contains(PolygonToCheck))
-					{
-						for (const FPolygonID AdjacentPolygon : Edge.ConnectedPolygons)
-						{
-							// Only add new polygons which haven't yet been added to the list. This prevents circular runs of polygons triggering infinite loops.
-							PolygonsToCheck.AddUnique(AdjacentPolygon);
-						}
-					}
-				}
-			}
-		}
-	}
-
-} //End of namespace MeshDescriptionPrototype
-
-void FMeshDescriptionHelper::CreateNormals(UMeshDescription* MeshDescription, FMeshDescriptionHelper::ETangentOptions TangentOptions, bool bComputeTangent)
-{
-	//For each vertex compute the normals for every connected edges that are smooth betwween hard edges
-	//         H   A    B
-	//          \  ||  /
-	//       G  -- ** -- C
-	//          // |  \
-	//         F   E    D
-	//
-	// The double ** are the vertex, the double line are hard edges, the single line are soft edge.
-	// A and F are hard, all other edges are soft. The goal is to compute two average normals one from
-	// A to F and a second from F to A. Then we can set the vertex instance normals accordingly.
-	// First normal(A to F) = Normalize(A+B+C+D+E+F)
-	// Second normal(F to A) = Normalize(F+G+H+A)
-	// We found the connected edge using the triangle that share edges
-
-	// @todo: provide an option to weight each contributing polygon normal according to the size of
-	// the angle it makes with the vertex being calculated. This means that triangulated faces whose
-	// internal edge meets the vertex doesn't get undue extra weight.
-	
-	const TVertexInstanceAttributeArray<FVector2D>& VertexUVs = MeshDescription->VertexInstanceAttributes().GetAttributes<FVector2D>(MeshAttribute::VertexInstance::TextureCoordinate, 0);
-	TVertexInstanceAttributeArray<FVector>& VertexNormals = MeshDescription->VertexInstanceAttributes().GetAttributes<FVector>(MeshAttribute::VertexInstance::Normal, 0);
-	TVertexInstanceAttributeArray<FVector>& VertexTangents = MeshDescription->VertexInstanceAttributes().GetAttributes<FVector>(MeshAttribute::VertexInstance::Tangent, 0);
-	TVertexInstanceAttributeArray<float>& VertexBinormalSigns = MeshDescription->VertexInstanceAttributes().GetAttributes<float>(MeshAttribute::VertexInstance::BinormalSign, 0);
-
-	TPolygonAttributeArray<FVector>& PolygonNormals = MeshDescription->PolygonAttributes().GetAttributes<FVector>(MeshAttribute::Polygon::Normal);
-	TPolygonAttributeArray<FVector>& PolygonTangents = MeshDescription->PolygonAttributes().GetAttributes<FVector>(MeshAttribute::Polygon::Tangent);
-	TPolygonAttributeArray<FVector>& PolygonBinormals = MeshDescription->PolygonAttributes().GetAttributes<FVector>(MeshAttribute::Polygon::Binormal);
-
-	TMap<FPolygonID, FVertexInfo> VertexInfoMap;
-	VertexInfoMap.Reserve(20);
-	//Iterate all vertex to compute normals for all vertex instance
-	for (const FVertexID VertexID : MeshDescription->Vertices().GetElementIDs())
-	{
-		VertexInfoMap.Reset();
-		
-		bool bPointHasAllTangents = true;
-		//Fill the VertexInfoMap
-		for (const FEdgeID EdgeID : MeshDescription->GetVertexConnectedEdges(VertexID))
-		{
-			for (const FPolygonID PolygonID : MeshDescription->GetEdgeConnectedPolygons(EdgeID))
-			{
-				FVertexInfo& VertexInfo = VertexInfoMap.FindOrAdd(PolygonID);
-				int32 EdgeIndex = VertexInfo.EdgeIDs.AddUnique(EdgeID);
-				if (VertexInfo.PolygonID == FPolygonID::Invalid)
-				{
-					VertexInfo.PolygonID = PolygonID;
-					for (const FVertexInstanceID VertexInstanceID : MeshDescription->GetPolygonPerimeterVertexInstances(PolygonID))
-					{
-						if (MeshDescription->GetVertexInstanceVertex(VertexInstanceID) == VertexID)
-						{
-							VertexInfo.VertexInstanceID = VertexInstanceID;
-							VertexInfo.UVs = VertexUVs[VertexInstanceID];
-							bPointHasAllTangents &= !VertexNormals[VertexInstanceID].IsNearlyZero() && !VertexTangents[VertexInstanceID].IsNearlyZero();
-							break;
-						}
-					}
-				}
-			}
-		}
-
-		if (bPointHasAllTangents)
-		{
-			continue;
-		}
-
-		//Make sure we consume all our vertex instance
-		check(VertexInfoMap.Num() == MeshDescription->GetVertexVertexInstances(VertexID).Num());
-
-		//Build all group by recursively traverse all polygon connected to the vertex
-		TArray<TArray<FPolygonID>> Groups;
-		TArray<FPolygonID> ConsumedPolygon;
-		for (auto Kvp : VertexInfoMap)
-		{
-			if (ConsumedPolygon.Contains(Kvp.Key))
-			{
-				continue;
-			}
-
-			int32 CurrentGroupIndex = Groups.AddZeroed();
-			TArray<FPolygonID>& CurrentGroup = Groups[CurrentGroupIndex];
-			TArray<FPolygonID> PolygonQueue;
-			PolygonQueue.Add(Kvp.Key); //Use a queue to avoid recursive function
-			while(PolygonQueue.Num() > 0)
-			{
-				FPolygonID CurrentPolygonID = PolygonQueue.Pop(true);
-				FVertexInfo& CurrentVertexInfo = VertexInfoMap.FindOrAdd(CurrentPolygonID);
-				CurrentGroup.AddUnique(CurrentVertexInfo.PolygonID);
-				ConsumedPolygon.AddUnique(CurrentVertexInfo.PolygonID);
-				const TEdgeAttributeArray<bool>& EdgeHardnesses = MeshDescription->EdgeAttributes().GetAttributes<bool>(MeshAttribute::Edge::IsHard);
-				for (const FEdgeID EdgeID : CurrentVertexInfo.EdgeIDs)
-				{
-					if (EdgeHardnesses[EdgeID])
-					{
-						//End of the group
-						continue;
-					}
-					for (const FPolygonID PolygonID : MeshDescription->GetEdgeConnectedPolygons(EdgeID))
-					{
-						if (PolygonID == CurrentVertexInfo.PolygonID)
-						{
-							continue;
-						}
-						//Add this polygon to the group
-						FVertexInfo& OtherVertexInfo = VertexInfoMap.FindOrAdd(PolygonID);
-						//Do not repeat polygons
-						if (!ConsumedPolygon.Contains(OtherVertexInfo.PolygonID))
-						{
-							PolygonQueue.Add(PolygonID);
-						}
-					}
-				}
-			}
-		}
-
-		//Smooth every connected group
-		ConsumedPolygon.Reset();
-		for(const TArray<FPolygonID>& Group : Groups)
-		{
-			//Compute tangents data
-			TMap<FVector2D, FVector> GroupTangent;
-			TMap<FVector2D, FVector> GroupBiNormal;
-
-			TArray<FVertexInstanceID> VertexInstanceInGroup;
-			FVector GroupNormal(0.0f);
-			for(const FPolygonID PolygonID : Group)
-			{
-#ifdef ENABLE_NTB_CHECK
-				check(!ConsumedPolygon.Contains(PolygonID));
-#endif
-				ConsumedPolygon.Add(PolygonID);
-				VertexInstanceInGroup.Add(VertexInfoMap[PolygonID].VertexInstanceID);
-				GroupNormal += PolygonNormals[PolygonID];
-				if (bComputeTangent)
-				{
-					const FVector2D UVs = VertexInfoMap[PolygonID].UVs;
-					bool CreateGroup = (!GroupTangent.Contains(UVs));
-					FVector& GroupTangentValue = GroupTangent.FindOrAdd(UVs);
-					FVector& GroupBiNormalValue = GroupBiNormal.FindOrAdd(UVs);
-					GroupTangentValue = CreateGroup ? PolygonTangents[PolygonID] : GroupTangentValue + PolygonTangents[PolygonID];
-					GroupBiNormalValue = CreateGroup ? PolygonBinormals[PolygonID] : GroupBiNormalValue + PolygonBinormals[PolygonID];
-				}
-			}
-
-			//////////////////////////////////////////////////////////////////////////
-			//Apply the group to the Mesh
-			GroupNormal.Normalize();
-			if (bComputeTangent)
-			{
-				for (auto Kvp : GroupTangent)
-				{
-					FVector& GroupTangentValue = GroupTangent.FindOrAdd(Kvp.Key);
-					GroupTangentValue.Normalize();
-				}
-				for (auto Kvp : GroupBiNormal)
-				{
-					FVector& GroupBiNormalValue = GroupBiNormal.FindOrAdd(Kvp.Key);
-					GroupBiNormalValue.Normalize();
-				}
-			}
-			//Apply the average NTB on all Vertex instance
-			for (const FVertexInstanceID VertexInstanceID : VertexInstanceInGroup)
-			{
-				const FVector2D VertexUV = VertexUVs[VertexInstanceID];
-			
-				if (VertexNormals[VertexInstanceID].IsNearlyZero(SMALL_NUMBER))
-				{
-					VertexNormals[VertexInstanceID] = GroupNormal;
-				}
-#ifdef ENABLE_NTB_CHECK
-				check(!VertexNormals[VertexInstanceID].IsNearlyZero(SMALL_NUMBER));
-#endif
-				if (bComputeTangent)
-				{
-					//Avoid changing the original group value
-					FVector GroupTangentValue = GroupTangent[VertexUV];
-					FVector GroupBiNormalValue = GroupBiNormal[VertexUV];
-
-					if (!VertexTangents[VertexInstanceID].IsNearlyZero(SMALL_NUMBER))
-					{
-						GroupTangentValue = VertexTangents[VertexInstanceID];
-					}
-#ifdef ENABLE_NTB_CHECK
-					check(!GroupTangentValue.IsNearlyZero(SMALL_NUMBER))
-#endif
-					FVector BiNormal(0.0f);
-					if (!VertexNormals[VertexInstanceID].IsNearlyZero(SMALL_NUMBER) && !VertexTangents[VertexInstanceID].IsNearlyZero(SMALL_NUMBER))
-					{
-						BiNormal = FVector::CrossProduct(VertexNormals[VertexInstanceID], VertexTangents[VertexInstanceID]).GetSafeNormal() * VertexBinormalSigns[VertexInstanceID];
-					}
-					if (!BiNormal.IsNearlyZero(SMALL_NUMBER))
-					{
-						GroupBiNormalValue = BiNormal;
-					}
-#ifdef ENABLE_NTB_CHECK
-					check(!GroupBiNormalValue.IsNearlyZero(SMALL_NUMBER));
-#endif
-					// Gram-Schmidt orthogonalization
-					GroupBiNormalValue -= GroupTangentValue * (GroupTangentValue | GroupBiNormalValue);
-					GroupBiNormalValue.Normalize();
-					
-					GroupTangentValue -= VertexNormals[VertexInstanceID] * (VertexNormals[VertexInstanceID] | GroupTangentValue);
-					GroupTangentValue.Normalize();
-					
-					GroupBiNormalValue -= VertexNormals[VertexInstanceID] * (VertexNormals[VertexInstanceID] | GroupBiNormalValue);
-					GroupBiNormalValue.Normalize();
-#ifdef ENABLE_NTB_CHECK
-					check(!GroupTangentValue.IsNearlyZero(SMALL_NUMBER));
-					check(!GroupBiNormalValue.IsNearlyZero(SMALL_NUMBER));
-#endif
-					//Set the value
-					VertexTangents[VertexInstanceID] = GroupTangentValue;
-					//If the BiNormal is zero set the sign to 1.0f
-					VertexBinormalSigns[VertexInstanceID] = GetBasisDeterminantSign(GroupTangentValue, GroupBiNormalValue, VertexNormals[VertexInstanceID]);
-
-				}
-			}
-		}
-	}
-}
-
-void FMeshDescriptionHelper::CreateMikktTangents(UMeshDescription* MeshDescription, FMeshDescriptionHelper::ETangentOptions TangentOptions)
-{
-	bool bIgnoreDegenerateTriangles = (TangentOptions & FMeshDescriptionHelper::ETangentOptions::IgnoreDegenerateTriangles) != 0;
-
-	// we can use mikktspace to calculate the tangents
-	SMikkTSpaceInterface MikkTInterface;
-	MikkTInterface.m_getNormal = MeshDescriptionMikktSpaceInterface::MikkGetNormal;
-	MikkTInterface.m_getNumFaces = MeshDescriptionMikktSpaceInterface::MikkGetNumFaces;
-	MikkTInterface.m_getNumVerticesOfFace = MeshDescriptionMikktSpaceInterface::MikkGetNumVertsOfFace;
-	MikkTInterface.m_getPosition = MeshDescriptionMikktSpaceInterface::MikkGetPosition;
-	MikkTInterface.m_getTexCoord = MeshDescriptionMikktSpaceInterface::MikkGetTexCoord;
-	MikkTInterface.m_setTSpaceBasic = MeshDescriptionMikktSpaceInterface::MikkSetTSpaceBasic;
-	MikkTInterface.m_setTSpace = nullptr;
-
-	SMikkTSpaceContext MikkTContext;
-	MikkTContext.m_pInterface = &MikkTInterface;
-	MikkTContext.m_pUserData = (void*)(MeshDescription);
-	MikkTContext.m_bIgnoreDegenerates = bIgnoreDegenerateTriangles;
-	genTangSpaceDefault(&MikkTContext);
-}
-
-namespace MeshDescriptionMikktSpaceInterface
-{
-	int MikkGetNumFaces(const SMikkTSpaceContext* Context)
-	{
-		UMeshDescription *MeshDescription = (UMeshDescription*)(Context->m_pUserData);
-		return MeshDescription->Polygons().Num();
-	}
-
-	int MikkGetNumVertsOfFace(const SMikkTSpaceContext* Context, const int FaceIdx)
-	{
-		// All of our meshes are triangles.
-		return 3;
-	}
-
-	void MikkGetPosition(const SMikkTSpaceContext* Context, float Position[3], const int FaceIdx, const int VertIdx)
-	{
-		UMeshDescription *MeshDescription = (UMeshDescription*)(Context->m_pUserData);
-		const FVertexInstanceID VertexInstanceID(FaceIdx * 3 + VertIdx);
-		const FVertexID VertexID = MeshDescription->GetVertexInstanceVertex(VertexInstanceID);
-		const FVector& VertexPosition = MeshDescription->VertexAttributes().GetAttribute<FVector>(VertexID, MeshAttribute::Vertex::Position);
-		Position[0] = VertexPosition.X;
-		Position[1] = VertexPosition.Y;
-		Position[2] = VertexPosition.Z;
-	}
-
-	void MikkGetNormal(const SMikkTSpaceContext* Context, float Normal[3], const int FaceIdx, const int VertIdx)
-	{
-		UMeshDescription *MeshDescription = (UMeshDescription*)(Context->m_pUserData);
-		const FVertexInstanceID VertexInstanceID(FaceIdx * 3 + VertIdx);
-		const FVector& VertexNormal = MeshDescription->VertexInstanceAttributes().GetAttribute<FVector>(VertexInstanceID, MeshAttribute::VertexInstance::Normal);
-		Normal[0] = VertexNormal.X;
-		Normal[1] = VertexNormal.Y;
-		Normal[2] = VertexNormal.Z;
-	}
-
-	void MikkSetTSpaceBasic(const SMikkTSpaceContext* Context, const float Tangent[3], const float BitangentSign, const int FaceIdx, const int VertIdx)
-	{
-		UMeshDescription *MeshDescription = (UMeshDescription*)(Context->m_pUserData);
-		const FVertexInstanceID VertexInstanceID(FaceIdx * 3 + VertIdx);
-		const FVector VertexTangent(Tangent[0], Tangent[1], Tangent[2]);
-		MeshDescription->VertexInstanceAttributes().SetAttribute<FVector>(VertexInstanceID, MeshAttribute::VertexInstance::Tangent, 0, VertexTangent);
-		MeshDescription->VertexInstanceAttributes().SetAttribute<float>(VertexInstanceID, MeshAttribute::VertexInstance::BinormalSign, 0, -BitangentSign);
-	}
-
-	void MikkGetTexCoord(const SMikkTSpaceContext* Context, float UV[2], const int FaceIdx, const int VertIdx)
-	{
-		UMeshDescription *MeshDescription = (UMeshDescription*)(Context->m_pUserData);
-		const FVertexInstanceID VertexInstanceID(FaceIdx * 3 + VertIdx);
-		const FVector2D& TexCoord = MeshDescription->VertexInstanceAttributes().GetAttribute<FVector2D>(VertexInstanceID, MeshAttribute::VertexInstance::TextureCoordinate, 0);
-		UV[0] = TexCoord.X;
-		UV[1] = TexCoord.Y;
-	}
-}
-
-//////////////////////////////////////////////////////////////////////////
-//Converters
-
-void FMeshDescriptionHelper::ConvertHardEdgesToSmoothGroup(const UMeshDescription* SourceMeshDescription, struct FRawMesh &DestinationRawMesh)
-{
-	TMap<FPolygonID, uint32> PolygonSmoothGroup;
-	PolygonSmoothGroup.Reserve(SourceMeshDescription->Polygons().Num());
-	TArray<bool> ConsumedPolygons;
-	ConsumedPolygons.AddZeroed(SourceMeshDescription->Polygons().Num());
-
-	TMap < FPolygonID, uint32> PolygonAvoidances;
-
-	const TEdgeAttributeArray<bool>& EdgeHardnesses = SourceMeshDescription->EdgeAttributes().GetAttributes<bool>(MeshAttribute::Edge::IsHard);
-
-	for (const FPolygonID PolygonID : SourceMeshDescription->Polygons().GetElementIDs())
-	{
-		if (ConsumedPolygons[PolygonID.GetValue()])
-		{
-			continue;
-		}
-		TArray<FPolygonID> ConnectedPolygons;
-		TArray<FPolygonID> LastConnectedPolygons;
-		ConnectedPolygons.Add(PolygonID);
-		LastConnectedPolygons.Add(FPolygonID::Invalid);
-		while (ConnectedPolygons.Num() > 0)
-		{
-			check(LastConnectedPolygons.Num() == ConnectedPolygons.Num());
-			FPolygonID LastPolygonID = LastConnectedPolygons.Pop(true);
-			FPolygonID CurrentPolygonID = ConnectedPolygons.Pop(true);
-			if (ConsumedPolygons[CurrentPolygonID.GetValue()])
-			{
-				continue;
-			}
-			TArray<FPolygonID> SoftEdgeNeigbors;
-			uint32& SmoothGroup = PolygonSmoothGroup.FindOrAdd(CurrentPolygonID);
-			uint32 AvoidSmoothGroup = 0;
-			uint32 NeighborSmoothGroup = 0;
-			const uint32 LastSmoothGroupValue = (LastPolygonID == FPolygonID::Invalid) ? 0 : PolygonSmoothGroup[LastPolygonID];
-			TArray<FEdgeID> PolygonEdges;
-			SourceMeshDescription->GetPolygonEdges(CurrentPolygonID, PolygonEdges);
-			for (const FEdgeID& EdgeID : PolygonEdges)
-			{
-				bool bIsHardEdge = EdgeHardnesses[EdgeID];
-				const TArray<FPolygonID>& EdgeConnectedPolygons = SourceMeshDescription->GetEdgeConnectedPolygons(EdgeID);
-				for (const FPolygonID& EdgePolygonID : EdgeConnectedPolygons)
-				{
-					if (EdgePolygonID == CurrentPolygonID)
-					{
-						continue;
-					}
-					uint32 SmoothValue = 0;
-					if (PolygonSmoothGroup.Contains(EdgePolygonID))
-					{
-						SmoothValue = PolygonSmoothGroup[EdgePolygonID];
-					}
-
-					if (bIsHardEdge) //Hard Edge
-					{
-						AvoidSmoothGroup |= SmoothValue;
-					}
-					else
-					{
-						NeighborSmoothGroup |= SmoothValue;
-						//Put all none hard edge polygon in the next iteration
-						if (!ConsumedPolygons[EdgePolygonID.GetValue()])
-						{
-							ConnectedPolygons.Add(EdgePolygonID);
-							LastConnectedPolygons.Add(CurrentPolygonID);
-						}
-						else
-						{
-							SoftEdgeNeigbors.Add(EdgePolygonID);
-						}
-					}
-				}
-			}
-
-			if (AvoidSmoothGroup != 0)
-			{
-				PolygonAvoidances.FindOrAdd(CurrentPolygonID) = AvoidSmoothGroup;
-				//find neighbor avoidance
-				for (FPolygonID& NeighborID : SoftEdgeNeigbors)
-				{
-					if(!PolygonAvoidances.Contains(NeighborID))
-					{
-						continue;
-					}
-					AvoidSmoothGroup |= PolygonAvoidances[NeighborID];
-				}
-				uint32 NewSmoothGroup = 1;
-				while ((NewSmoothGroup & AvoidSmoothGroup) != 0 && NewSmoothGroup < MAX_uint32)
-				{
-					//Shift the smooth group
-					NewSmoothGroup = NewSmoothGroup << 1;
-				}
-				SmoothGroup = NewSmoothGroup;
-				//Aply to all neighboard
-				for(FPolygonID& NeighborID : SoftEdgeNeigbors)
-				{
-					PolygonSmoothGroup[NeighborID] |= NewSmoothGroup;
-				}
-			}
-			else if (NeighborSmoothGroup != 0)
-			{
-				SmoothGroup |= LastSmoothGroupValue | NeighborSmoothGroup;
-			}
-			else
-			{
-				SmoothGroup = 1;
-			}
-			ConsumedPolygons[CurrentPolygonID.GetValue()] = true;
-		}
-	}
-	//Now we have to put the data into the RawMesh
-	int32 TriangleIndex = 0;
-	for (const FPolygonID PolygonID : SourceMeshDescription->Polygons().GetElementIDs())
-	{
-		uint32 PolygonSmoothValue = PolygonSmoothGroup[PolygonID];
-		const TArray<FMeshTriangle>& Triangles = SourceMeshDescription->GetPolygonTriangles(PolygonID);
-		for (const FMeshTriangle& MeshTriangle : Triangles)
-		{
-			DestinationRawMesh.FaceSmoothingMasks[TriangleIndex++] = PolygonSmoothValue;
-		}
-	}
-}
-
-
-
-
-void FMeshDescriptionHelper::ConvertSmoothGroupToHardEdges(const struct FRawMesh &SourceRawMesh, UMeshDescription* DestinationMeshDescription)
-{
-	TEdgeAttributeArray<bool>& EdgeHardnesses = DestinationMeshDescription->EdgeAttributes().GetAttributes<bool>(MeshAttribute::Edge::IsHard);
-
-	TArray<bool> ConsumedPolygons;
-	ConsumedPolygons.AddZeroed(DestinationMeshDescription->Polygons().Num());
-	for (const FPolygonID PolygonID : DestinationMeshDescription->Polygons().GetElementIDs())
-	{
-		if (ConsumedPolygons[PolygonID.GetValue()])
-		{
-			continue;
-		}
-		TArray<FPolygonID> ConnectedPolygons;
-		ConnectedPolygons.Add(PolygonID);
-		while (ConnectedPolygons.Num() > 0)
-		{
-			FPolygonID CurrentPolygonID = ConnectedPolygons.Pop(true);
-			int32 CurrentPolygonIDValue = CurrentPolygonID.GetValue();
-			check(SourceRawMesh.FaceSmoothingMasks.IsValidIndex(CurrentPolygonIDValue));
-			const uint32 ReferenceSmoothGroup = SourceRawMesh.FaceSmoothingMasks[CurrentPolygonIDValue];
-			TArray<FEdgeID> PolygonEdges;
-			DestinationMeshDescription->GetPolygonEdges(CurrentPolygonID, PolygonEdges);
-			for (const FEdgeID& EdgeID : PolygonEdges)
-			{
-				const bool bIsHardEdge = EdgeHardnesses[EdgeID];
-				if (bIsHardEdge)
-				{
-					continue;
-				}
-				const TArray<FPolygonID>& EdgeConnectedPolygons = DestinationMeshDescription->GetEdgeConnectedPolygons(EdgeID);
-				for (const FPolygonID& EdgePolygonID : EdgeConnectedPolygons)
-				{
-					int32 EdgePolygonIDValue = EdgePolygonID.GetValue();
-					if (EdgePolygonID == CurrentPolygonID || ConsumedPolygons[EdgePolygonIDValue])
-					{
-						continue;
-					}
-					check(SourceRawMesh.FaceSmoothingMasks.IsValidIndex(EdgePolygonIDValue));
-					const uint32 TestSmoothGroup = SourceRawMesh.FaceSmoothingMasks[EdgePolygonIDValue];
-					if ((TestSmoothGroup & ReferenceSmoothGroup) == 0)
-					{
-						EdgeHardnesses[EdgeID] = true;
-						break;
-					}
-					else
-					{
-						ConnectedPolygons.Add(EdgePolygonID);
-					}
-				}
-			}
-			ConsumedPolygons[CurrentPolygonID.GetValue()] = true;
-		}
-	}
-}
-
-void FMeshDescriptionHelper::ConverToRawMesh(const UMeshDescription* SourceMeshDescription, struct FRawMesh &DestinationRawMesh)
-{
-	static BuildStatisticManager::FStatisticData StatisticData;
-	BuildStatisticManager::FBuildStatisticScope StatScope(TEXT("ConvertToRawMesh"), StatisticData);
-	DestinationRawMesh.Empty();
-
-	//Gather all array data
-	const TVertexAttributeArray<FVector>& VertexPositions = SourceMeshDescription->VertexAttributes().GetAttributes<FVector>(MeshAttribute::Vertex::Position);
-
-	const TVertexInstanceAttributeArray<FVector>& VertexInstanceNormals = SourceMeshDescription->VertexInstanceAttributes().GetAttributes<FVector>(MeshAttribute::VertexInstance::Normal);
-	const TVertexInstanceAttributeArray<FVector>& VertexInstanceTangents = SourceMeshDescription->VertexInstanceAttributes().GetAttributes<FVector>(MeshAttribute::VertexInstance::Tangent);
-	const TVertexInstanceAttributeArray<float>& VertexInstanceBinormalSigns = SourceMeshDescription->VertexInstanceAttributes().GetAttributes<float>(MeshAttribute::VertexInstance::BinormalSign);
-	const TVertexInstanceAttributeArray<FVector4>& VertexInstanceColors = SourceMeshDescription->VertexInstanceAttributes().GetAttributes<FVector4>(MeshAttribute::VertexInstance::Color);
-	const TVertexInstanceAttributeIndicesArray<FVector2D>& VertexInstanceUVs = SourceMeshDescription->VertexInstanceAttributes().GetAttributesSet<FVector2D>(MeshAttribute::VertexInstance::TextureCoordinate);
-
-	const TPolygonGroupAttributeArray<int>& PolygonGroupMaterialIndex = SourceMeshDescription->PolygonGroupAttributes().GetAttributes<int>(MeshAttribute::PolygonGroup::MaterialIndex);
-
-	DestinationRawMesh.VertexPositions.AddZeroed(SourceMeshDescription->Vertices().Num());
-	for (const FVertexID& VertexID : SourceMeshDescription->Vertices().GetElementIDs())
-	{
-		int32 VertexIDValue = VertexID.GetValue();
-		DestinationRawMesh.VertexPositions[VertexIDValue] = VertexPositions[VertexID];
-	}
-	int32 VertexInstanceNumber = SourceMeshDescription->VertexInstances().Num();
-	DestinationRawMesh.WedgeColors.AddZeroed(VertexInstanceNumber);
-	DestinationRawMesh.WedgeIndices.AddZeroed(VertexInstanceNumber);
-	DestinationRawMesh.WedgeTangentX.AddZeroed(VertexInstanceNumber);
-	DestinationRawMesh.WedgeTangentY.AddZeroed(VertexInstanceNumber);
-	DestinationRawMesh.WedgeTangentZ.AddZeroed(VertexInstanceNumber);
-	int32 ExistingUVCount = VertexInstanceUVs.GetNumIndices();
-	for (int32 UVIndex = 0; UVIndex < ExistingUVCount; ++UVIndex)
-	{
-		DestinationRawMesh.WedgeTexCoords[UVIndex].AddZeroed(VertexInstanceNumber);
-	}
-
-	int32 TriangleNumber = 0;
-	for (const FPolygonID& PolygonID : SourceMeshDescription->Polygons().GetElementIDs())
-	{
-		TriangleNumber += SourceMeshDescription->GetPolygonTriangles(PolygonID).Num();
-	}
-	DestinationRawMesh.FaceMaterialIndices.AddZeroed(TriangleNumber);
-	DestinationRawMesh.FaceSmoothingMasks.AddZeroed(TriangleNumber);
-
-	for (const FPolygonID& PolygonID : SourceMeshDescription->Polygons().GetElementIDs())
-	{
-		const FPolygonGroupID& PolygonGroupID = SourceMeshDescription->GetPolygonPolygonGroup(PolygonID);
-		int32 PolygonIDValue = PolygonID.GetValue();
-		const TArray<FMeshTriangle>& Triangles = SourceMeshDescription->GetPolygonTriangles(PolygonID);
-		for (const FMeshTriangle& MeshTriangle : Triangles)
-		{
-			DestinationRawMesh.FaceMaterialIndices[PolygonIDValue] = PolygonGroupMaterialIndex[PolygonGroupID];
-			DestinationRawMesh.FaceSmoothingMasks[PolygonIDValue] = 0; //Conversion of soft/hard to smooth mask is done after the geometry is converted
-			for (int32 Corner = 0; Corner < 3; ++Corner)
-			{
-				const FVertexInstanceID VertexInstanceID = MeshTriangle.GetVertexInstanceID(Corner);
-				const int32 VertexInstanceIDValue = VertexInstanceID.GetValue();
-				DestinationRawMesh.WedgeColors[VertexInstanceIDValue] = FLinearColor(VertexInstanceColors[VertexInstanceID]).ToFColor(true);
-				DestinationRawMesh.WedgeIndices[VertexInstanceIDValue] = SourceMeshDescription->GetVertexInstanceVertex(VertexInstanceID).GetValue();
-				DestinationRawMesh.WedgeTangentX[VertexInstanceIDValue] = VertexInstanceTangents[VertexInstanceID];
-				DestinationRawMesh.WedgeTangentY[VertexInstanceIDValue] = FVector::CrossProduct(VertexInstanceNormals[VertexInstanceID], VertexInstanceTangents[VertexInstanceID]).GetSafeNormal() * VertexInstanceBinormalSigns[VertexInstanceID];
-				DestinationRawMesh.WedgeTangentZ[VertexInstanceIDValue] = VertexInstanceNormals[VertexInstanceID];
-				for (int32 UVIndex = 0; UVIndex < ExistingUVCount; ++UVIndex)
-				{
-					DestinationRawMesh.WedgeTexCoords[UVIndex][VertexInstanceIDValue] = VertexInstanceUVs.GetArrayForIndex(UVIndex)[VertexInstanceID];
-				}
-			}
-		}
-	}
-	//Convert the smoothgroup
-	ConvertHardEdgesToSmoothGroup(SourceMeshDescription, DestinationRawMesh);
-}
-
-void FMeshDescriptionHelper::ConverFromRawMesh(const struct FRawMesh &SourceRawMesh, UMeshDescription* DestinationMeshDescription)
-{
-	check(DestinationMeshDescription != nullptr);
-	static BuildStatisticManager::FStatisticData StatisticData;
-	BuildStatisticManager::FBuildStatisticScope StatScope(TEXT("ConverFromRawMesh"), StatisticData);
-	DestinationMeshDescription->Empty();
-
-	DestinationMeshDescription->ReserveNewVertices(SourceRawMesh.VertexPositions.Num());
-	DestinationMeshDescription->ReserveNewVertexInstances(SourceRawMesh.WedgeIndices.Num());
-	DestinationMeshDescription->ReserveNewPolygons(SourceRawMesh.WedgeIndices.Num() / 3);
-	DestinationMeshDescription->ReserveNewEdges(SourceRawMesh.WedgeIndices.Num() * 2 / 3); //Approximately 2 edges per polygons
-	//Gather all array data
-	TVertexAttributeArray<FVector>& VertexPositions = DestinationMeshDescription->VertexAttributes().GetAttributes<FVector>(MeshAttribute::Vertex::Position);
-
-	TVertexInstanceAttributeArray<FVector>& VertexInstanceNormals = DestinationMeshDescription->VertexInstanceAttributes().GetAttributes<FVector>(MeshAttribute::VertexInstance::Normal);
-	TVertexInstanceAttributeArray<FVector>& VertexInstanceTangents = DestinationMeshDescription->VertexInstanceAttributes().GetAttributes<FVector>(MeshAttribute::VertexInstance::Tangent);
-	TVertexInstanceAttributeArray<float>& VertexInstanceBinormalSigns = DestinationMeshDescription->VertexInstanceAttributes().GetAttributes<float>(MeshAttribute::VertexInstance::BinormalSign);
-	TVertexInstanceAttributeArray<FVector4>& VertexInstanceColors = DestinationMeshDescription->VertexInstanceAttributes().GetAttributes<FVector4>(MeshAttribute::VertexInstance::Color);
-	TVertexInstanceAttributeIndicesArray<FVector2D>& VertexInstanceUVs = DestinationMeshDescription->VertexInstanceAttributes().GetAttributesSet<FVector2D>(MeshAttribute::VertexInstance::TextureCoordinate);
-
-	TPolygonGroupAttributeArray<FName>& PolygonGroupImportedMaterialSlotNames = DestinationMeshDescription->PolygonGroupAttributes().GetAttributes<FName>(MeshAttribute::PolygonGroup::ImportedMaterialSlotName);
-	TPolygonGroupAttributeArray<int>& PolygonGroupMaterialIndex = DestinationMeshDescription->PolygonGroupAttributes().GetAttributes<int>(MeshAttribute::PolygonGroup::MaterialIndex);
-
-	int32 NumTexCoords = 0;
-	int32 MaxTexCoords = FMath::Min<int32>(MAX_MESH_TEXTURE_COORDS, MAX_STATIC_TEXCOORDS);
-	TArray<int32> TextureCoordinnateRemapIndex;
-	TextureCoordinnateRemapIndex.AddZeroed(MaxTexCoords);
-	for (int32 TextureCoordinnateIndex = 0; TextureCoordinnateIndex < MaxTexCoords; ++TextureCoordinnateIndex)
-	{
-		TextureCoordinnateRemapIndex[TextureCoordinnateIndex] = INDEX_NONE;
-		if (SourceRawMesh.WedgeTexCoords[TextureCoordinnateIndex].Num() == SourceRawMesh.WedgeIndices.Num())
-		{
-			TextureCoordinnateRemapIndex[TextureCoordinnateIndex] = NumTexCoords;
-			NumTexCoords++;
-		}
-	}
-	VertexInstanceUVs.SetNumIndices(NumTexCoords);
-	for (int32 VertexIndex = 0; VertexIndex < SourceRawMesh.VertexPositions.Num(); ++VertexIndex)
-	{
-		FVertexID VertexID = DestinationMeshDescription->CreateVertex();
-		VertexPositions[VertexID] = SourceRawMesh.VertexPositions[VertexIndex];
-	}
-
-	bool bHasColors = SourceRawMesh.WedgeColors.Num() > 0;
-	bool bHasTangents = SourceRawMesh.WedgeTangentX.Num() > 0 && SourceRawMesh.WedgeTangentY.Num() > 0;
-	bool bHasNormals = SourceRawMesh.WedgeTangentZ.Num() > 0;
-
-	TArray<FPolygonGroupID> PolygonGroups;
-
-	//Triangles
-	int32 TriangleCount = SourceRawMesh.WedgeIndices.Num() / 3;
-	for (int32 TriangleIndex = 0; TriangleIndex < TriangleCount; ++TriangleIndex)
-	{
-		int32 VerticeIndexBase = TriangleIndex * 3;
-
-		//PolygonGroup
-		FPolygonGroupID PolygonGroupID = FPolygonGroupID::Invalid;
-		int32 MaterialIndex = SourceRawMesh.FaceMaterialIndices[TriangleIndex];
-		for (const FPolygonGroupID& SearchPolygonGroupID : DestinationMeshDescription->PolygonGroups().GetElementIDs())
-		{
-			if (PolygonGroupMaterialIndex[SearchPolygonGroupID] == MaterialIndex)
-			{
-				PolygonGroupID = SearchPolygonGroupID;
-			}
-		}
-		if (PolygonGroupID == FPolygonGroupID::Invalid)
-		{
-			PolygonGroupID = DestinationMeshDescription->CreatePolygonGroup();
-			PolygonGroupMaterialIndex[PolygonGroupID] = MaterialIndex;
-			PolygonGroupImportedMaterialSlotNames[PolygonGroupID] = FName(*FString::Printf(TEXT("MaterialSlot_%d"), MaterialIndex));
-			PolygonGroups.Add(PolygonGroupID);
-		}
-
-		for (int32 Corner = 0; Corner < 3; ++Corner)
-		{
-			int32 VerticeIndex = VerticeIndexBase + Corner;
-			FVertexID VertexID(SourceRawMesh.WedgeIndices[VerticeIndex]);
-			FVertexInstanceID VertexInstanceID = DestinationMeshDescription->CreateVertexInstance(VertexID);
-			VertexInstanceColors[VertexInstanceID] = bHasColors ? FLinearColor::FromSRGBColor(SourceRawMesh.WedgeColors[VerticeIndex]) : FLinearColor::White;
-			VertexInstanceTangents[VertexInstanceID] = bHasTangents ? SourceRawMesh.WedgeTangentX[VerticeIndex] : FVector(ForceInitToZero);
-			VertexInstanceBinormalSigns[VertexInstanceID] = bHasTangents ? GetBasisDeterminantSign(SourceRawMesh.WedgeTangentX[VerticeIndex].GetSafeNormal(), SourceRawMesh.WedgeTangentY[VerticeIndex].GetSafeNormal(), SourceRawMesh.WedgeTangentZ[VerticeIndex].GetSafeNormal()) : 0.0f;
-			VertexInstanceNormals[VertexInstanceID] = bHasNormals ? SourceRawMesh.WedgeTangentZ[VerticeIndex] : FVector(ForceInitToZero);
-			for (int32 TextureCoordinnateIndex = 0; TextureCoordinnateIndex < NumTexCoords; ++TextureCoordinnateIndex)
-			{
-				int32 TextureCoordIndex = TextureCoordinnateRemapIndex[TextureCoordinnateIndex];
-				if (TextureCoordIndex == INDEX_NONE)
-				{
-					continue;
-				}
-				TMeshAttributeArray<FVector2D, FVertexInstanceID>& Uvs = VertexInstanceUVs.GetArrayForIndex(TextureCoordIndex);
-				Uvs[VertexInstanceID] = SourceRawMesh.WedgeTexCoords[TextureCoordinnateIndex][VerticeIndex];
-			}
-		}
-		
-		//Create the polygon edges
-		TArray<UMeshDescription::FContourPoint> Contours;
-		// Add the edges of this triangle
-		for (uint32 Corner = 0; Corner < 3; ++Corner)
-		{
-			int32 ContourPointIndex = Contours.AddDefaulted();
-			UMeshDescription::FContourPoint& ContourPoint = Contours[ContourPointIndex];
-			//Find the matching edge ID
-			int32 CornerIndices[2];
-			CornerIndices[0] = (Corner + 0) % 3;
-			CornerIndices[1] = (Corner + 1) % 3;
-
-			FVertexID EdgeVertexIDs[2];
-			EdgeVertexIDs[0] = DestinationMeshDescription->GetVertexInstanceVertex(FVertexInstanceID(VerticeIndexBase + CornerIndices[0]));
-			EdgeVertexIDs[1] = DestinationMeshDescription->GetVertexInstanceVertex(FVertexInstanceID(VerticeIndexBase + CornerIndices[1]));
-
-			FEdgeID MatchEdgeId = DestinationMeshDescription->GetVertexPairEdge(EdgeVertexIDs[0], EdgeVertexIDs[1]);
-			if (MatchEdgeId == FEdgeID::Invalid)
-			{
-				MatchEdgeId = DestinationMeshDescription->CreateEdge(EdgeVertexIDs[0], EdgeVertexIDs[1]);
-			}
-			ContourPoint.EdgeID = MatchEdgeId;
-			ContourPoint.VertexInstanceID = FVertexInstanceID(VerticeIndexBase + CornerIndices[0]);
-		}
-
-		const FPolygonID NewPolygonID = DestinationMeshDescription->CreatePolygon(PolygonGroupID, Contours);
-		int32 NewTriangleIndex = DestinationMeshDescription->GetPolygonTriangles(NewPolygonID).AddDefaulted();
-		FMeshTriangle& NewTriangle = DestinationMeshDescription->GetPolygonTriangles(NewPolygonID)[NewTriangleIndex];
-		for (int32 Corner = 0; Corner < 3; ++Corner)
-		{
-			FVertexInstanceID VertexInstanceID(VerticeIndexBase + Corner);
-			NewTriangle.SetVertexInstanceID(Corner, VertexInstanceID);
-		}
-	}
-	CreatePolygonNTB(DestinationMeshDescription, 0.0f);
-
-	if (!bHasNormals || !bHasTangents)
-	{
-		//Create the missing normals and tangents
-		if (!bHasNormals)
-		{
-			CreateNormals(DestinationMeshDescription, ETangentOptions::BlendOverlappingNormals, false);
-		}
-		CreateMikktTangents(DestinationMeshDescription, ETangentOptions::BlendOverlappingNormals);
-	}
-
-	ConvertSmoothGroupToHardEdges(SourceRawMesh, DestinationMeshDescription);
-}
\ No newline at end of file
diff --git a/Engine/Source/Runtime/MeshBuilder/Private/MeshDescriptionHelper.h b/Engine/Source/Runtime/MeshBuilder/Private/MeshDescriptionHelper.h
--- a/Engine/Source/Runtime/MeshBuilder/Private/MeshDescriptionHelper.h
+++ b/Engine/Source/Runtime/MeshBuilder/Private/MeshDescriptionHelper.h
@@ -21,21 +21,6 @@ enum
 class FMeshDescriptionHelper
 {
 public:
-	enum ETangentOptions
-	{
-		None = 0,
-		BlendOverlappingNormals = 0x1,
-		IgnoreDegenerateTriangles = 0x2,
-		UseMikkTSpace = 0x4,
-	};
-
-	enum class ELightmapUVVersion : int32
-	{
-		BitByBit = 0,
-		Segments = 1,
-		SmallChartPacking = 2,
-		Latest = SmallChartPacking
-	};
 
 	FMeshDescriptionHelper(FMeshBuildSettings* InBuildSettings, const UMeshDescription* InOriginalMeshDescription);
 
@@ -52,22 +37,11 @@ public:
 
 	const TMultiMap<int32, int32>& GetOverlappingCorners() const { return OverlappingCorners; }
 
-	/** Convert this mesh description into the old FRawMesh format*/
-	static void ConverToRawMesh(const UMeshDescription* SourceMeshDescription, struct FRawMesh &DestinationRawMesh);
-	/** Convert old FRawMesh format to MeshDescription*/
-	static void ConverFromRawMesh(const struct FRawMesh &SourceRawMesh, UMeshDescription* DestinationMeshDescription);
-
-	static void CreatePolygonNTB(UMeshDescription* MeshDescription, float ComparisonThreshold);
-	static void CreateNormals(UMeshDescription* MeshDescription, ETangentOptions TangentOptions, bool bComputeTangent);
-	static void CreateMikktTangents(UMeshDescription* MeshDescription, ETangentOptions TangentOptions);
 private:
 
 	//////////////////////////////////////////////////////////////////////////
 	//PRIVATE function declarations
 
-	static void ConvertHardEdgesToSmoothGroup(const UMeshDescription* SourceMeshDescription, struct FRawMesh &DestinationRawMesh);
-	static void ConvertSmoothGroupToHardEdges(const struct FRawMesh &SourceRawMesh, UMeshDescription* DestinationMeshDescription);
-
 	//////////////////////////////////////////////////////////////////////////
 	//PRIVATE class members
 
diff --git a/Engine/Source/Runtime/MeshBuilder/Public/MeshBuilder.h b/Engine/Source/Runtime/MeshBuilder/Public/MeshBuilder.h
--- a/Engine/Source/Runtime/MeshBuilder/Public/MeshBuilder.h
+++ b/Engine/Source/Runtime/MeshBuilder/Public/MeshBuilder.h
@@ -3,7 +3,6 @@
 #pragma once
 
 #include "CoreMinimal.h"
-#include "Materials/MaterialInterface.h"
 
 /**
  * Abstract class which is the base class of all builder.
@@ -23,18 +22,9 @@ private:
 
 };
 
-class MESHBUILDER_API FMeshDescriptionOperations
+class MESHBUILDER_API FMeshBuilderOperations
 {
 public:
-	/** Compute the normals tangents and binormals for the mesh description. */
-	static void ComputeMeshNTBs(class UMeshDescription* MeshDescription, const struct FMeshBuildSettings& BuildSettings);
-
-	/** Convert this mesh description into the old FRawMesh format. */
-	static void ConverToRawMesh(const class UMeshDescription* SourceMeshDescription, struct FRawMesh &DestinationRawMesh);
-
-	/** Convert old FRawMesh format to MeshDescription. */
-	static void ConverFromRawMesh(const struct FRawMesh &SourceRawMesh, class UMeshDescription* DestinationMeshDescription);
-
 	/** Create some UVs from the specified mesh description data. */
 	static bool GenerateUniqueUVsForStaticMesh(const class UMeshDescription* MeshDescription, int32 TextureResolution, TArray<FVector2D>& OutTexCoords);
 };
