Fix win32 warning

diff --git a/Engine/Plugins/Runtime/EditableMesh/Source/EditableMesh/EditableMesh.cpp b/Engine/Plugins/Runtime/EditableMesh/Source/EditableMesh/EditableMesh.cpp
--- a/Engine/Plugins/Runtime/EditableMesh/Source/EditableMesh/EditableMesh.cpp
+++ b/Engine/Plugins/Runtime/EditableMesh/Source/EditableMesh/EditableMesh.cpp
@@ -5379,7 +5379,7 @@ void UEditableMesh::ExtrudePolygons( const TArray<FPolygonID>& PolygonIDs, const
 						// would avoid the problem where extruding two polygons that are connected only by a vertex are not extruded
 						// separately.
 						const bool bVertexIsSharedByAnEdgeOfAnotherSelectedPolygon = false;
-						if( ExtrudedEdgeVertexIDPtr != nullptr && !bVertexIsSharedByAnEdgeOfAnotherSelectedPolygon )
+						if( ExtrudedEdgeVertexIDPtr != nullptr)// && !bVertexIsSharedByAnEdgeOfAnotherSelectedPolygon )
 						{
 							ExtrudedEdgeVertexIDs[ EdgeVertexNumber ] = *ExtrudedEdgeVertexIDPtr;
 						}
diff --git a/Engine/Source/Runtime/MeshBuilder/MeshBuilder.Build.cs b/Engine/Source/Runtime/MeshBuilder/MeshBuilder.Build.cs
--- a/Engine/Source/Runtime/MeshBuilder/MeshBuilder.Build.cs
+++ b/Engine/Source/Runtime/MeshBuilder/MeshBuilder.Build.cs
@@ -17,18 +17,25 @@ namespace UnrealBuildTool.Rules
                     "Engine",
                     "RenderCore",
                     "MeshDescription",
-                    "MeshDescriptionOperations",
-                    "MeshReductionInterface",
-                    "RenderCore",
-                    "RawMesh"
+                    "RenderCore"
                 }
 			);
 
-            AddEngineThirdPartyPrivateStaticDependencies(Target, "nvTriStrip");
-            AddEngineThirdPartyPrivateStaticDependencies(Target, "ForsythTriOptimizer");
-            AddEngineThirdPartyPrivateStaticDependencies(Target, "nvTessLib");
-            AddEngineThirdPartyPrivateStaticDependencies(Target, "QuadricMeshReduction");
-
+            if (Target.bBuildEditor)
+            {
+                PrivateDependencyModuleNames.AddRange(
+                    new string[]
+                    {
+                        "MeshDescriptionOperations",
+                        "MeshReductionInterface",
+                        "RawMesh"
+                    }
+                );
+				AddEngineThirdPartyPrivateStaticDependencies(Target, "nvTriStrip");
+            	AddEngineThirdPartyPrivateStaticDependencies(Target, "ForsythTriOptimizer");
+	            AddEngineThirdPartyPrivateStaticDependencies(Target, "nvTessLib");
+                AddEngineThirdPartyPrivateStaticDependencies(Target, "QuadricMeshReduction");
+            }
             if (Target.bCompileSimplygon == true)
             {
                 AddEngineThirdPartyPrivateDynamicDependencies(Target, "SimplygonMeshReduction");
diff --git a/Engine/Source/Runtime/MeshBuilder/Private/MeshDescriptionHelper.cpp b/Engine/Source/Runtime/MeshBuilder/Private/MeshDescriptionHelper.cpp
--- a/Engine/Source/Runtime/MeshBuilder/Private/MeshDescriptionHelper.cpp
+++ b/Engine/Source/Runtime/MeshBuilder/Private/MeshDescriptionHelper.cpp
@@ -130,12 +130,15 @@ UMeshDescription* FMeshDescriptionHelper::GetRenderMeshDescription(UObject* Owne
 			VertexInstanceUVs.SetNumIndices(BuildSettings->DstLightmapIndex + 1);
 			BuildSettings->DstLightmapIndex = NumIndices;
 		}
-
 		FMeshDescriptionOperations::CreateLightMapUVLayout(RenderMeshDescription,
 			BuildSettings->SrcLightmapIndex,
 			BuildSettings->DstLightmapIndex,
 			BuildSettings->MinLightmapResolution,
+#if WITH_EDITORONLY_DATA
 			(FMeshDescriptionOperations::ELightmapUVVersion)(StaticMesh->LightmapUVVersion),
+#else
+			FMeshDescriptionOperations::ELightmapUVVersion::SmallChartPacking,
+#endif
 			OverlappingCorners);
 	}
 
diff --git a/Engine/Source/Runtime/MeshBuilder/Private/StaticMeshBuilder.cpp b/Engine/Source/Runtime/MeshBuilder/Private/StaticMeshBuilder.cpp
--- a/Engine/Source/Runtime/MeshBuilder/Private/StaticMeshBuilder.cpp
+++ b/Engine/Source/Runtime/MeshBuilder/Private/StaticMeshBuilder.cpp
@@ -15,6 +15,7 @@
 
 DEFINE_LOG_CATEGORY(LogStaticMeshBuilder);
 
+#if WITH_EDITOR
 //////////////////////////////////////////////////////////////////////////
 //Local functions definition
 bool IsOrphanedVertex(const class UMeshDescription *MeshDescription, const FVertexID VertexID);
@@ -630,4 +631,5 @@ void BuildAllBufferOptimizations(FStaticMeshLODResources& StaticMeshLOD, const F
 		);
 		StaticMeshLOD.AdjacencyIndexBuffer.SetIndices(AdjacencyIndices, IndexBufferStride);
 	}
-}
\ No newline at end of file
+}
+#endif //WITH_EDITOR
\ No newline at end of file
diff --git a/Engine/Source/Runtime/MeshBuilder/Public/StaticMeshBuilder.h b/Engine/Source/Runtime/MeshBuilder/Public/StaticMeshBuilder.h
--- a/Engine/Source/Runtime/MeshBuilder/Public/StaticMeshBuilder.h
+++ b/Engine/Source/Runtime/MeshBuilder/Public/StaticMeshBuilder.h
@@ -9,6 +9,7 @@ DECLARE_LOG_CATEGORY_EXTERN(LogStaticMeshBuilder, Log, All);
 
 class MESHBUILDER_API FStaticMeshBuilder : public FMeshBuilder
 {
+#if WITH_EDITOR
 public:
 	FStaticMeshBuilder();
 
@@ -22,5 +23,6 @@ private:
 
 	/** Used to refresh all components in the scene that may be using a mesh we're editing */
 	TSharedPtr<class FStaticMeshComponentRecreateRenderStateContext> RecreateRenderStateContext;
+#endif //WITH_EDITOR
 };
 
